PR #5864: Plastic SCM Source Control provider plugin for Unreal Engine 4 (Contributed by [REDACTED_PERSON])

diff --git a/Engine/Plugins/Developer/PlasticSourceControl/PlasticSourceControl.uplugin b/Engine/Plugins/Developer/PlasticSourceControl/PlasticSourceControl.uplugin
new file mode 100644
--- /dev/null
+++ b/Engine/Plugins/Developer/PlasticSourceControl/PlasticSourceControl.uplugin
@@ -0,0 +1,24 @@
+{
+	"FileVersion": 3,
+	"Version": 40,
+	"VersionName": "1.4.6",
+	"FriendlyName": "Plastic SCM",
+	"Description": "Plastic source control management",
+	"Category": "Source Control",
+	"CreatedBy": "SRombauts for Codice Software.",
+	"CreatedByURL": "http://srombauts.github.com",
+	"DocsURL": "https://www.plasticscm.com/documentation.html",
+	"MarketplaceURL": "",
+	"SupportURL": "https://www.plasticscm.com/support.html",
+	"EnabledByDefault": true,
+	"CanContainContent": false,
+	"IsBetaVersion": false,
+	"Installed": false,
+	"Modules": [
+		{
+			"Name": "PlasticSourceControl",
+			"Type": "Editor",
+			"LoadingPhase": "Default"
+		}
+	]
+}
\ No newline at end of file
diff --git a/Engine/Plugins/Developer/PlasticSourceControl/Source/PlasticSourceControl/PlasticSourceControl.Build.cs b/Engine/Plugins/Developer/PlasticSourceControl/Source/PlasticSourceControl/PlasticSourceControl.Build.cs
new file mode 100644
--- /dev/null
+++ b/Engine/Plugins/Developer/PlasticSourceControl/Source/PlasticSourceControl/PlasticSourceControl.Build.cs
@@ -0,0 +1,27 @@
+// Copyright 1998-2019 Epic Games, Inc. All Rights Reserved.
+
+using UnrealBuildTool;
+
+public class PlasticSourceControl : ModuleRules
+{
+	public PlasticSourceControl(ReadOnlyTargetRules Target) : base(Target)
+	{
+		PrivateDependencyModuleNames.AddRange(
+			new string[]
+			{
+				"Core",
+				"CoreUObject",
+				"Slate",
+				"SlateCore",
+				"InputCore",
+				"EditorStyle",
+				"UnrealEd",
+				"LevelEditor",
+				"SourceControl",
+				"XmlParser",
+				"Projects",
+				"AssetRegistry",
+			}
+		);
+	}
+}
diff --git a/Engine/Plugins/Developer/PlasticSourceControl/Source/PlasticSourceControl/Private/IPlasticSourceControlWorker.h b/Engine/Plugins/Developer/PlasticSourceControl/Source/PlasticSourceControl/Private/IPlasticSourceControlWorker.h
new file mode 100644
--- /dev/null
+++ b/Engine/Plugins/Developer/PlasticSourceControl/Source/PlasticSourceControl/Private/IPlasticSourceControlWorker.h
@@ -0,0 +1,27 @@
+// Copyright 1998-2019 Epic Games, Inc. All Rights Reserved.
+
+#pragma once
+
+#include "CoreMinimal.h"
+
+class IPlasticSourceControlWorker
+{
+public:
+	/**
+	 * Name describing the work that this worker does. Used for factory method hookup.
+	 */
+	virtual FName GetName() const = 0;
+
+	/**
+	 * Function that actually does the work. Can be executed on another thread.
+	 */
+	virtual bool Execute( class FPlasticSourceControlCommand& InCommand ) = 0;
+
+	/**
+	 * Updates the state of any items after completion (if necessary). This is always executed on the main thread.
+	 * @returns true if states were updated
+	 */
+	virtual bool UpdateStates() = 0;
+};
+
+typedef TSharedRef<IPlasticSourceControlWorker, ESPMode::ThreadSafe> FPlasticSourceControlWorkerRef;
\ No newline at end of file
diff --git a/Engine/Plugins/Developer/PlasticSourceControl/Source/PlasticSourceControl/Private/PlasticSourceControlCommand.cpp b/Engine/Plugins/Developer/PlasticSourceControl/Source/PlasticSourceControl/Private/PlasticSourceControlCommand.cpp
new file mode 100644
--- /dev/null
+++ b/Engine/Plugins/Developer/PlasticSourceControl/Source/PlasticSourceControl/Private/PlasticSourceControlCommand.cpp
@@ -0,0 +1,44 @@
+// Copyright 1998-2019 Epic Games, Inc. All Rights Reserved.
+
+#include "PlasticSourceControlCommand.h"
+#include "PlasticSourceControlModule.h"
+#include "ISourceControlOperation.h"
+#include "IPlasticSourceControlWorker.h"
+#include "Modules/ModuleManager.h"
+
+
+FPlasticSourceControlCommand::FPlasticSourceControlCommand(const TSharedRef<class ISourceControlOperation, ESPMode::ThreadSafe>& InOperation, const TSharedRef<class IPlasticSourceControlWorker, ESPMode::ThreadSafe>& InWorker, const FSourceControlOperationComplete& InOperationCompleteDelegate)
+	: Operation(InOperation)
+	, Worker(InWorker)
+	, OperationCompleteDelegate(InOperationCompleteDelegate)
+	, bExecuteProcessed(0)
+	, bCommandSuccessful(false)
+	, bConnectionDropped(false)
+	, bAutoDelete(true)
+	, Concurrency(EConcurrency::Synchronous)
+{
+	// grab the providers settings here, so we don't access them once the worker thread is launched
+	check(IsInGameThread());
+	const FPlasticSourceControlModule& PlasticSourceControl = FModuleManager::LoadModuleChecked<FPlasticSourceControlModule>( "PlasticSourceControl" );
+	PathToWorkspaceRoot = PlasticSourceControl.GetProvider().GetPathToWorkspaceRoot();
+	ChangesetNumber = PlasticSourceControl.GetProvider().GetChangesetNumber();
+}
+
+bool FPlasticSourceControlCommand::DoWork()
+{
+	bCommandSuccessful = Worker->Execute(*this);
+	FPlatformAtomics::InterlockedExchange(&bExecuteProcessed, 1);
+
+	return bCommandSuccessful;
+}
+
+void FPlasticSourceControlCommand::Abandon()
+{
+	FPlatformAtomics::InterlockedExchange(&bExecuteProcessed, 1);
+}
+
+void FPlasticSourceControlCommand::DoThreadedWork()
+{
+	Concurrency = EConcurrency::Asynchronous;
+	DoWork();
+}
diff --git a/Engine/Plugins/Developer/PlasticSourceControl/Source/PlasticSourceControl/Private/PlasticSourceControlCommand.h b/Engine/Plugins/Developer/PlasticSourceControl/Source/PlasticSourceControl/Private/PlasticSourceControlCommand.h
new file mode 100644
--- /dev/null
+++ b/Engine/Plugins/Developer/PlasticSourceControl/Source/PlasticSourceControl/Private/PlasticSourceControlCommand.h
@@ -0,0 +1,89 @@
+// Copyright 1998-2019 Epic Games, Inc. All Rights Reserved.
+
+#pragma once
+
+#include "CoreMinimal.h"
+#include "Misc/IQueuedWork.h"
+#include "ISourceControlProvider.h"
+
+/**
+ * Used to execute Plastic commands multi-threaded.
+ */
+class FPlasticSourceControlCommand : public IQueuedWork
+{
+public:
+
+	FPlasticSourceControlCommand(const TSharedRef<class ISourceControlOperation, ESPMode::ThreadSafe>& InOperation, const TSharedRef<class IPlasticSourceControlWorker, ESPMode::ThreadSafe>& InWorker, const FSourceControlOperationComplete& InOperationCompleteDelegate = FSourceControlOperationComplete() );
+
+	/**
+	 * This is where the real thread work is done. All work that is done for
+	 * this queued object should be done from within the call to this function.
+	 */
+	bool DoWork();
+
+	/**
+	 * Tells the queued work that it is being abandoned so that it can do
+	 * per object clean up as needed. This will only be called if it is being
+	 * abandoned before completion. NOTE: This requires the object to delete
+	 * itself using whatever heap it was allocated in.
+	 */
+	virtual void Abandon() override;
+
+	/**
+	 * This method is also used to tell the object to cleanup but not before
+	 * the object has finished it's work.
+	 */ 
+	virtual void DoThreadedWork() override;
+
+public:
+	/** Path to the root of the Plastic workspace: can be the GameDir itself, or any parent directory (found by the "Connect" operation) */
+	FString PathToWorkspaceRoot;
+
+	/** Operation we want to perform - contains outward-facing parameters & results */
+	TSharedRef<class ISourceControlOperation, ESPMode::ThreadSafe> Operation;
+
+	/** The object that will actually do the work */
+	TSharedRef<class IPlasticSourceControlWorker, ESPMode::ThreadSafe> Worker;
+
+	/** Delegate to notify when this operation completes */
+	FSourceControlOperationComplete OperationCompleteDelegate;
+
+	/**If true, this command has been processed by the source control thread*/
+	volatile int32 bExecuteProcessed;
+
+	/**If true, the source control command succeeded*/
+	bool bCommandSuccessful;
+
+	/**If true, the source control connection was dropped while this command was being executed*/
+	bool bConnectionDropped;
+
+	/** Plastic current workspace */
+	FString WorkspaceName;
+
+	/** Plastic current repository */
+	FString RepositoryName;
+
+	/** Plastic current server URL */
+	FString ServerUrl;
+
+	/** Name of the current branch */
+	FString BranchName;
+
+	/** Current Changeset Number */
+	int32 ChangesetNumber;
+
+	/** If true, this command will be automatically cleaned up in Tick() */
+	bool bAutoDelete;
+
+	/** Whether we are running multi-treaded in the background, or blocking the main thread */
+	EConcurrency::Type Concurrency;
+
+	/** Files to perform this operation on */
+	TArray< FString > Files;
+
+	/**Info and/or warning message message storage*/
+	TArray< FString > InfoMessages;
+
+	/**Potential error message storage*/
+	TArray< FString > ErrorMessages;
+};
\ No newline at end of file
diff --git a/Engine/Plugins/Developer/PlasticSourceControl/Source/PlasticSourceControl/Private/PlasticSourceControlMenu.cpp b/Engine/Plugins/Developer/PlasticSourceControl/Source/PlasticSourceControl/Private/PlasticSourceControlMenu.cpp
new file mode 100644
--- /dev/null
+++ b/Engine/Plugins/Developer/PlasticSourceControl/Source/PlasticSourceControl/Private/PlasticSourceControlMenu.cpp
@@ -0,0 +1,447 @@
+// Copyright 1998-2019 Epic Games, Inc. All Rights Reserved.
+
+#include "PlasticSourceControlMenu.h"
+
+#include "PlasticSourceControlModule.h"
+#include "PlasticSourceControlProvider.h"
+#include "PlasticSourceControlOperations.h"
+
+#include "ISourceControlModule.h"
+#include "ISourceControlOperation.h"
+#include "SourceControlOperations.h"
+
+#include "LevelEditor.h"
+#include "Widgets/Notifications/SNotificationList.h"
+#include "Framework/Notifications/NotificationManager.h"
+#include "Framework/MultiBox/MultiBoxBuilder.h"
+#include "Misc/MessageDialog.h"
+#include "EditorStyleSet.h"
+
+#include "PackageTools.h"
+#include "FileHelpers.h"
+
+#include "Logging/MessageLog.h"
+
+static const FName PlasticSourceControlMenuTabName("PlasticSourceControlMenu");
+
+#define LOCTEXT_NAMESPACE "PlasticSourceControl"
+
+void FPlasticSourceControlMenu::Register()
+{
+	// Register the extension with the level editor
+	FLevelEditorModule* LevelEditorModule = FModuleManager::GetModulePtr<FLevelEditorModule>("LevelEditor");
+	if (LevelEditorModule)
+	{
+		FLevelEditorModule::FLevelEditorMenuExtender ViewMenuExtender = FLevelEditorModule::FLevelEditorMenuExtender::CreateRaw(this, &FPlasticSourceControlMenu::OnExtendLevelEditorViewMenu);
+		auto& MenuExtenders = LevelEditorModule->GetAllLevelEditorToolbarSourceControlMenuExtenders();
+		MenuExtenders.Add(ViewMenuExtender);
+		ViewMenuExtenderHandle = MenuExtenders.Last().GetHandle();
+	}
+}
+
+void FPlasticSourceControlMenu::Unregister()
+{
+	// Unregister the level editor extensions
+	FLevelEditorModule* LevelEditorModule = FModuleManager::GetModulePtr<FLevelEditorModule>("LevelEditor");
+	if (LevelEditorModule)
+	{
+		LevelEditorModule->GetAllLevelEditorToolbarSourceControlMenuExtenders().RemoveAll([=](const FLevelEditorModule::FLevelEditorMenuExtender& Extender) { return Extender.GetHandle() == ViewMenuExtenderHandle; });
+	}
+}
+
+bool FPlasticSourceControlMenu::IsSourceControlConnected() const
+{
+	const ISourceControlProvider& Provider = ISourceControlModule::Get().GetProvider();
+	return Provider.IsEnabled() && Provider.IsAvailable();
+}
+
+/// Prompt to save or discard all packages
+bool FPlasticSourceControlMenu::SaveDirtyPackages()
+{
+	const bool bPromptUserToSave = true;
+	const bool bSaveMapPackages = true;
+	const bool bSaveContentPackages = true;
+	const bool bFastSave = false;
+	const bool bNotifyNoPackagesSaved = false;
+	const bool bCanBeDeclined = true; // If the user clicks "don't save" this will continue and lose their changes
+	bool bHadPackagesToSave = false;
+
+	bool bSaved = FEditorFileUtils::SaveDirtyPackages(bPromptUserToSave, bSaveMapPackages, bSaveContentPackages, bFastSave, bNotifyNoPackagesSaved, bCanBeDeclined, &bHadPackagesToSave);
+
+	// bSaved can be true if the user selects to not save an asset by unchecking it and clicking "save"
+	if (bSaved)
+	{
+		TArray<UPackage*> DirtyPackages;
+		FEditorFileUtils::GetDirtyWorldPackages(DirtyPackages);
+		FEditorFileUtils::GetDirtyContentPackages(DirtyPackages);
+		bSaved = DirtyPackages.Num() == 0;
+	}
+
+	return bSaved;
+}
+
+/// Find all packages in Content directory
+TArray<FString> FPlasticSourceControlMenu::ListAllPackages()
+{
+	TArray<FString> PackageRelativePaths;
+	FPackageName::FindPackagesInDirectory(PackageRelativePaths, *FPaths::ConvertRelativePathToFull(FPaths::ProjectContentDir()));
+
+	TArray<FString> PackageNames;
+	PackageNames.Reserve(PackageRelativePaths.Num());
+	for (const FString& Path : PackageRelativePaths)
+	{
+		FString PackageName;
+		FString FailureReason;
+		if (FPackageName::TryConvertFilenameToLongPackageName(Path, PackageName, &FailureReason))
+		{
+			PackageNames.Add(PackageName);
+		}
+		else
+		{
+			FMessageLog("SourceControl").Error(FText::FromString(FailureReason));
+		}
+	}
+
+	return PackageNames;
+}
+
+/// Unkink all loaded packages to allow to update them
+TArray<UPackage*> FPlasticSourceControlMenu::UnlinkPackages(const TArray<FString>& InPackageNames)
+{
+	TArray<UPackage*> LoadedPackages;
+
+	// Inspired from ContentBrowserUtils::SyncPathsFromSourceControl()
+	if (InPackageNames.Num() > 0)
+	{
+		// Form a list of loaded packages to reload...
+		LoadedPackages.Reserve(InPackageNames.Num());
+		for (const FString& PackageName : InPackageNames)
+		{
+			UPackage* Package = FindPackage(nullptr, *PackageName);
+			if (Package)
+			{
+				LoadedPackages.Emplace(Package);
+
+				// Detach the linkers of any loaded packages so that SCC can overwrite the files...
+				if (!Package->IsFullyLoaded())
+				{
+					FlushAsyncLoading();
+					Package->FullyLoad();
+				}
+				ResetLoaders(Package);
+			}
+		}
+		UE_LOG(LogSourceControl, Log, TEXT("Reseted Loader for %d Packages"), LoadedPackages.Num());
+	}
+
+	return LoadedPackages;
+}
+
+void FPlasticSourceControlMenu::ReloadPackages(TArray<UPackage*>& InPackagesToReload)
+{
+	UE_LOG(LogSourceControl, Log, TEXT("Reloading %d Packages..."), InPackagesToReload.Num());
+
+	// Syncing may have deleted some packages, so we need to unload those rather than re-load them...
+	TArray<UPackage*> PackagesToUnload;
+	InPackagesToReload.RemoveAll([&](UPackage* InPackage) -> bool
+	{
+		const FString PackageExtension = InPackage->ContainsMap() ? FPackageName::GetMapPackageExtension() : FPackageName::GetAssetPackageExtension();
+		const FString PackageFilename = FPackageName::LongPackageNameToFilename(InPackage->GetName(), PackageExtension);
+		if (!FPaths::FileExists(PackageFilename))
+		{
+			PackagesToUnload.Emplace(InPackage);
+			return true; // remove package
+		}
+		return false; // keep package
+	});
+
+	// Hot-reload the new packages...
+	PackageTools::ReloadPackages(InPackagesToReload);
+
+	// Unload any deleted packages...
+	PackageTools::UnloadPackages(PackagesToUnload);
+}
+
+void FPlasticSourceControlMenu::SyncProjectClicked()
+{
+	if (!OperationInProgressNotification.IsValid())
+	{
+		const bool bSaved = SaveDirtyPackages();
+		if (bSaved)
+		{
+			// Find and Unlink all packages in Content directory to allow to update them
+			PackagesToReload = UnlinkPackages(ListAllPackages());
+
+			// Launch a "Sync" operation
+			FPlasticSourceControlModule& PlasticSourceControl = FModuleManager::LoadModuleChecked<FPlasticSourceControlModule>("PlasticSourceControl");
+			FPlasticSourceControlProvider& Provider = PlasticSourceControl.GetProvider();
+			TSharedRef<FSync, ESPMode::ThreadSafe> SyncOperation = ISourceControlOperation::Create<FSync>();
+			const ECommandResult::Type Result = Provider.Execute(SyncOperation, TArray<FString>(), EConcurrency::Asynchronous, FSourceControlOperationComplete::CreateRaw(this, &FPlasticSourceControlMenu::OnSourceControlOperationComplete));
+			if (Result == ECommandResult::Succeeded)
+			{
+				// Display an ongoing notification during the whole operation (packages will be reloaded at the completion of the operation)
+				DisplayInProgressNotification(SyncOperation->GetInProgressString());
+			}
+			else
+			{
+				// Report failure with a notification and Reload all packages
+				DisplayFailureNotification(SyncOperation->GetName());
+				ReloadPackages(PackagesToReload);
+			}
+		}
+		else
+		{
+			FMessageLog SourceControlLog("SourceControl");
+			SourceControlLog.Warning(LOCTEXT("SourceControlMenu_Sync_Unsaved", "Save All Assets before attempting to Sync!"));
+			SourceControlLog.Notify();
+		}
+	}
+	else
+	{
+		FMessageLog SourceControlLog("SourceControl");
+		SourceControlLog.Warning(LOCTEXT("SourceControlMenu_InProgress", "Source control operation already in progress"));
+		SourceControlLog.Notify();
+	}
+}
+
+void FPlasticSourceControlMenu::RevertUnchangedClicked()
+{
+	if (!OperationInProgressNotification.IsValid())
+	{
+		// Launch a "RevertUnchanged" Operation
+		FPlasticSourceControlModule& PlasticSourceControl = FModuleManager::LoadModuleChecked<FPlasticSourceControlModule>("PlasticSourceControl");
+		FPlasticSourceControlProvider& Provider = PlasticSourceControl.GetProvider();
+		TSharedRef<FPlasticRevertUnchanged, ESPMode::ThreadSafe> RevertUnchangedOperation = ISourceControlOperation::Create<FPlasticRevertUnchanged>();
+		TArray<FString> WorkspaceRoot;
+		WorkspaceRoot.Add(Provider.GetPathToWorkspaceRoot()); // Revert the root of the workspace (needs an absolute path)
+		const ECommandResult::Type Result = Provider.Execute(RevertUnchangedOperation, WorkspaceRoot, EConcurrency::Asynchronous, FSourceControlOperationComplete::CreateRaw(this, &FPlasticSourceControlMenu::OnSourceControlOperationComplete));
+		if (Result == ECommandResult::Succeeded)
+		{
+			// Display an ongoing notification during the whole operation
+			DisplayInProgressNotification(RevertUnchangedOperation->GetInProgressString());
+		}
+		else
+		{
+			// Report failure with a notification
+			DisplayFailureNotification(RevertUnchangedOperation->GetName());
+		}
+	}
+	else
+	{
+		FMessageLog SourceControlLog("SourceControl");
+		SourceControlLog.Warning(LOCTEXT("SourceControlMenu_InProgress", "Source control operation already in progress"));
+		SourceControlLog.Notify();
+	}
+}
+
+void FPlasticSourceControlMenu::RevertAllClicked()
+{
+	if (!OperationInProgressNotification.IsValid())
+	{
+		// Ask the user before reverting all!
+		const FText DialogText(LOCTEXT("SourceControlMenu_AskRevertAll", "Revert all modifications into the workspace?"));
+		const EAppReturnType::Type Choice = FMessageDialog::Open(EAppMsgType::OkCancel, DialogText);
+		if (Choice == EAppReturnType::Ok)
+		{
+			const bool bSaved = SaveDirtyPackages();
+			if (bSaved)
+			{
+				// Find and Unlink all packages in Content directory to allow to update them
+				PackagesToReload = UnlinkPackages(ListAllPackages());
+
+				// Launch a "RevertAll" Operation
+				FPlasticSourceControlModule& PlasticSourceControl = FModuleManager::LoadModuleChecked<FPlasticSourceControlModule>("PlasticSourceControl");
+				FPlasticSourceControlProvider& Provider = PlasticSourceControl.GetProvider();
+				TSharedRef<FPlasticRevertAll, ESPMode::ThreadSafe> RevertAllOperation = ISourceControlOperation::Create<FPlasticRevertAll>();
+				TArray<FString> WorkspaceRoot;
+				WorkspaceRoot.Add(Provider.GetPathToWorkspaceRoot()); // Revert the root of the workspace (needs an absolute path)
+				const ECommandResult::Type Result = Provider.Execute(RevertAllOperation, WorkspaceRoot, EConcurrency::Asynchronous, FSourceControlOperationComplete::CreateRaw(this, &FPlasticSourceControlMenu::OnSourceControlOperationComplete));
+				if (Result == ECommandResult::Succeeded)
+				{
+					// Display an ongoing notification during the whole operation
+					DisplayInProgressNotification(RevertAllOperation->GetInProgressString());
+				}
+				else
+				{
+					// Report failure with a notification and Reload all packages
+					DisplayFailureNotification(RevertAllOperation->GetName());
+					ReloadPackages(PackagesToReload);
+				}
+			}
+			else
+			{
+				FMessageLog SourceControlLog("SourceControl");
+				SourceControlLog.Warning(LOCTEXT("SourceControlMenu_Sync_Unsaved", "Save All Assets before attempting to Sync!"));
+				SourceControlLog.Notify();
+			}
+		}
+	}
+	else
+	{
+		FMessageLog SourceControlLog("SourceControl");
+		SourceControlLog.Warning(LOCTEXT("SourceControlMenu_InProgress", "Source control operation already in progress"));
+		SourceControlLog.Notify();
+	}
+}
+
+void FPlasticSourceControlMenu::RefreshClicked()
+{
+	if (!OperationInProgressNotification.IsValid())
+	{
+		// Launch an "UpdateStatus" Operation
+		FPlasticSourceControlModule& PlasticSourceControl = FModuleManager::LoadModuleChecked<FPlasticSourceControlModule>("PlasticSourceControl");
+		FPlasticSourceControlProvider& Provider = PlasticSourceControl.GetProvider();
+		TSharedRef<FUpdateStatus, ESPMode::ThreadSafe> RefreshOperation = ISourceControlOperation::Create<FUpdateStatus>();
+		RefreshOperation->SetCheckingAllFiles(true);
+		const ECommandResult::Type Result = Provider.Execute(RefreshOperation, TArray<FString>(), EConcurrency::Asynchronous, FSourceControlOperationComplete::CreateRaw(this, &FPlasticSourceControlMenu::OnSourceControlOperationComplete));
+		if (Result == ECommandResult::Succeeded)
+		{
+			// Display an ongoing notification during the whole operation
+			DisplayInProgressNotification(RefreshOperation->GetInProgressString());
+		}
+		else
+		{
+			// Report failure with a notification
+			DisplayFailureNotification(RefreshOperation->GetName());
+		}
+	}
+	else
+	{
+		FMessageLog SourceControlLog("SourceControl");
+		SourceControlLog.Warning(LOCTEXT("SourceControlMenu_InProgress", "Source control operation already in progress"));
+		SourceControlLog.Notify();
+	}
+}
+
+// Display an ongoing notification during the whole operation
+void FPlasticSourceControlMenu::DisplayInProgressNotification(const FText& InOperationInProgressString)
+{
+	if (!OperationInProgressNotification.IsValid())
+	{
+		FNotificationInfo Info(InOperationInProgressString);
+		Info.bFireAndForget = false;
+		Info.ExpireDuration = 0.0f;
+		Info.FadeOutDuration = 1.0f;
+		OperationInProgressNotification = FSlateNotificationManager::Get().AddNotification(Info);
+		if (OperationInProgressNotification.IsValid())
+		{
+			OperationInProgressNotification.Pin()->SetCompletionState(SNotificationItem::CS_Pending);
+		}
+	}
+}
+
+// Remove the ongoing notification at the end of the operation
+void FPlasticSourceControlMenu::RemoveInProgressNotification()
+{
+	if (OperationInProgressNotification.IsValid())
+	{
+		OperationInProgressNotification.Pin()->ExpireAndFadeout();
+		OperationInProgressNotification.Reset();
+	}
+}
+
+// Display a temporary success notification at the end of the operation
+void FPlasticSourceControlMenu::DisplaySucessNotification(const FName& InOperationName)
+{
+	const FText NotificationText = FText::Format(
+		LOCTEXT("SourceControlMenu_Success", "{0} operation was successful!"),
+		FText::FromName(InOperationName)
+	);
+	FNotificationInfo Info(NotificationText);
+	Info.bUseSuccessFailIcons = true;
+	Info.Image = FEditorStyle::GetBrush(TEXT("NotificationList.SuccessImage"));
+	FSlateNotificationManager::Get().AddNotification(Info);
+	UE_LOG(LogSourceControl, Log, TEXT("%s"), *NotificationText.ToString());
+}
+
+// Display a temporary failure notification at the end of the operation
+void FPlasticSourceControlMenu::DisplayFailureNotification(const FName& InOperationName)
+{
+	const FText NotificationText = FText::Format(
+		LOCTEXT("SourceControlMenu_Failure", "Error: {0} operation failed!"),
+		FText::FromName(InOperationName)
+	);
+	FNotificationInfo Info(NotificationText);
+	Info.ExpireDuration = 8.0f;
+	FSlateNotificationManager::Get().AddNotification(Info);
+	UE_LOG(LogSourceControl, Error, TEXT("%s"), *NotificationText.ToString());
+}
+
+void FPlasticSourceControlMenu::OnSourceControlOperationComplete(const FSourceControlOperationRef& InOperation, ECommandResult::Type InResult)
+{
+	RemoveInProgressNotification();
+
+	if ((InOperation->GetName() == "Sync") || (InOperation->GetName() == "RevertAll"))
+	{
+		// Reload packages that where unlinked at the beginning of the Sync operation
+		ReloadPackages(PackagesToReload);
+	}
+
+	// Report result with a notification
+	if (InResult == ECommandResult::Succeeded)
+	{
+		DisplaySucessNotification(InOperation->GetName());
+	}
+	else
+	{
+		DisplayFailureNotification(InOperation->GetName());
+	}
+}
+
+void FPlasticSourceControlMenu::AddMenuExtension(FMenuBuilder& Builder)
+{
+	Builder.AddMenuEntry(
+		LOCTEXT("PlasticSync",			"Sync/Update Workspace"),
+		LOCTEXT("PlasticSyncTooltip",	"Update all files in the workspace to the latest version."),
+		FSlateIcon(FEditorStyle::GetStyleSetName(), "SourceControl.Actions.Sync"),
+		FUIAction(
+			FExecuteAction::CreateRaw(this, &FPlasticSourceControlMenu::SyncProjectClicked),
+			FCanExecuteAction::CreateRaw(this, &FPlasticSourceControlMenu::IsSourceControlConnected)
+		)
+	);
+
+	Builder.AddMenuEntry(
+		LOCTEXT("PlasticRevertUnchanged",			"Revert Unchanged"),
+		LOCTEXT("PlasticRevertUnchangedTooltip",	"Revert checked-out but unchanged files in the workspace."),
+		FSlateIcon(FEditorStyle::GetStyleSetName(), "SourceControl.Actions.Revert"),
+		FUIAction(
+			FExecuteAction::CreateRaw(this, &FPlasticSourceControlMenu::RevertUnchangedClicked),
+			FCanExecuteAction()
+		)
+	);
+
+	Builder.AddMenuEntry(
+		LOCTEXT("PlasticRevertAll",			"Revert All"),
+		LOCTEXT("PlasticRevertAllTooltip",	"Revert all files in the workspace to their controlled/unchanged state."),
+		FSlateIcon(FEditorStyle::GetStyleSetName(), "SourceControl.Actions.Revert"),
+		FUIAction(
+			FExecuteAction::CreateRaw(this, &FPlasticSourceControlMenu::RevertAllClicked),
+			FCanExecuteAction()
+		)
+	);
+
+	Builder.AddMenuEntry(
+		LOCTEXT("PlasticRefresh",			"Refresh"),
+		LOCTEXT("PlasticRefreshTooltip",	"Update the source control status of all files in the workspace."),
+		FSlateIcon(FEditorStyle::GetStyleSetName(), "SourceControl.Actions.Refresh"),
+		FUIAction(
+			FExecuteAction::CreateRaw(this, &FPlasticSourceControlMenu::RefreshClicked),
+			FCanExecuteAction()
+		)
+	);
+}
+
+TSharedRef<FExtender> FPlasticSourceControlMenu::OnExtendLevelEditorViewMenu(const TSharedRef<FUICommandList> CommandList)
+{
+	TSharedRef<FExtender> Extender(new FExtender());
+
+	Extender->AddMenuExtension(
+		"SourceControlActions",
+		EExtensionHook::After,
+		nullptr,
+		FMenuExtensionDelegate::CreateRaw(this, &FPlasticSourceControlMenu::AddMenuExtension));
+
+	return Extender;
+}
+
+#undef LOCTEXT_NAMESPACE
diff --git a/Engine/Plugins/Developer/PlasticSourceControl/Source/PlasticSourceControl/Private/PlasticSourceControlMenu.h b/Engine/Plugins/Developer/PlasticSourceControl/Source/PlasticSourceControl/Private/PlasticSourceControlMenu.h
new file mode 100644
--- /dev/null
+++ b/Engine/Plugins/Developer/PlasticSourceControl/Source/PlasticSourceControl/Private/PlasticSourceControlMenu.h
@@ -0,0 +1,52 @@
+// Copyright 1998-2019 Epic Games, Inc. All Rights Reserved.
+
+#pragma once
+
+#include "CoreMinimal.h"
+#include "ISourceControlProvider.h"
+
+class FToolBarBuilder;
+class FMenuBuilder;
+
+/** Plastic SCM extension of the Source Control toolbar menu */
+class FPlasticSourceControlMenu
+{
+public:
+	void Register();
+	void Unregister();
+	
+	/** This functions will be bound to appropriate Command. */
+	void SyncProjectClicked();
+	void RevertUnchangedClicked();
+	void RevertAllClicked();
+	void RefreshClicked();
+
+private:
+	bool IsSourceControlConnected() const;
+
+	bool				SaveDirtyPackages();
+	TArray<FString>		ListAllPackages();
+	TArray<UPackage*>	UnlinkPackages(const TArray<FString>& InPackageNames);
+	void				ReloadPackages(TArray<UPackage*>& InPackagesToReload);
+
+	void AddMenuExtension(FMenuBuilder& Builder);
+
+	TSharedRef<class FExtender> OnExtendLevelEditorViewMenu(const TSharedRef<class FUICommandList> CommandList);
+
+	void DisplayInProgressNotification(const FText& InOperationInProgressString);
+	void RemoveInProgressNotification();
+	void DisplaySucessNotification(const FName& InOperationName);
+	void DisplayFailureNotification(const FName& InOperationName);
+
+private:
+	FDelegateHandle ViewMenuExtenderHandle;
+
+	/** Loaded packages to reload after a Sync or Revert operation */
+	TArray<UPackage*> PackagesToReload;
+
+	/** Current source control operation from extended menu if any */
+	TWeakPtr<class SNotificationItem> OperationInProgressNotification;
+
+	/** Delegate called when a source control operation has completed */
+	void OnSourceControlOperationComplete(const FSourceControlOperationRef& InOperation, ECommandResult::Type InResult);
+};
diff --git a/Engine/Plugins/Developer/PlasticSourceControl/Source/PlasticSourceControl/Private/PlasticSourceControlModule.cpp b/Engine/Plugins/Developer/PlasticSourceControl/Source/PlasticSourceControl/Private/PlasticSourceControlModule.cpp
new file mode 100644
--- /dev/null
+++ b/Engine/Plugins/Developer/PlasticSourceControl/Source/PlasticSourceControl/Private/PlasticSourceControlModule.cpp
@@ -0,0 +1,61 @@
+// Copyright 1998-2019 Epic Games, Inc. All Rights Reserved.
+
+#include "PlasticSourceControlModule.h"
+#include "PlasticSourceControlOperations.h"
+#include "Misc/App.h"
+#include "Features/IModularFeatures.h"
+
+#define LOCTEXT_NAMESPACE "PlasticSourceControl"
+
+template<typename Type>
+static TSharedRef<IPlasticSourceControlWorker, ESPMode::ThreadSafe> CreateWorker()
+{
+	return MakeShareable( new Type() );
+}
+
+void FPlasticSourceControlModule::StartupModule()
+{
+	// Register our operations (implemented in PlasticSourceControlOperations.cpp by subclassing from Engine\Source\Developer\SourceControl\Public\SourceControlOperations.h)
+	PlasticSourceControlProvider.RegisterWorker("Connect", FGetPlasticSourceControlWorker::CreateStatic( &CreateWorker<FPlasticConnectWorker>));
+	PlasticSourceControlProvider.RegisterWorker("CheckOut", FGetPlasticSourceControlWorker::CreateStatic(&CreateWorker<FPlasticCheckOutWorker>));
+	PlasticSourceControlProvider.RegisterWorker("UpdateStatus", FGetPlasticSourceControlWorker::CreateStatic( &CreateWorker<FPlasticUpdateStatusWorker>));
+	PlasticSourceControlProvider.RegisterWorker("MarkForAdd", FGetPlasticSourceControlWorker::CreateStatic(&CreateWorker<FPlasticMarkForAddWorker>));
+	PlasticSourceControlProvider.RegisterWorker("Delete", FGetPlasticSourceControlWorker::CreateStatic(&CreateWorker<FPlasticDeleteWorker>));
+	PlasticSourceControlProvider.RegisterWorker("Revert", FGetPlasticSourceControlWorker::CreateStatic(&CreateWorker<FPlasticRevertWorker>));
+	PlasticSourceControlProvider.RegisterWorker("RevertUnchanged", FGetPlasticSourceControlWorker::CreateStatic(&CreateWorker<FPlasticRevertUnchangedWorker>));
+	PlasticSourceControlProvider.RegisterWorker("RevertAll", FGetPlasticSourceControlWorker::CreateStatic(&CreateWorker<FPlasticRevertAllWorker>));
+	PlasticSourceControlProvider.RegisterWorker("MakeWorkspace", FGetPlasticSourceControlWorker::CreateStatic(&CreateWorker<FPlasticMakeWorkspaceWorker>));
+	PlasticSourceControlProvider.RegisterWorker("Sync", FGetPlasticSourceControlWorker::CreateStatic( &CreateWorker<FPlasticSyncWorker>));
+	PlasticSourceControlProvider.RegisterWorker("CheckIn", FGetPlasticSourceControlWorker::CreateStatic(&CreateWorker<FPlasticCheckInWorker>));
+	PlasticSourceControlProvider.RegisterWorker("Copy", FGetPlasticSourceControlWorker::CreateStatic(&CreateWorker<FPlasticCopyWorker>));
+	PlasticSourceControlProvider.RegisterWorker("Resolve", FGetPlasticSourceControlWorker::CreateStatic(&CreateWorker<FPlasticResolveWorker>));
+
+	// load our settings
+	PlasticSourceControlSettings.LoadSettings();
+
+	// Bind our source control provider to the editor
+	IModularFeatures::Get().RegisterModularFeature( "SourceControl", &PlasticSourceControlProvider );
+}
+
+void FPlasticSourceControlModule::ShutdownModule()
+{
+	// shut down the provider, as this module is going away
+	PlasticSourceControlProvider.Close();
+
+	// unbind provider from editor
+	IModularFeatures::Get().UnregisterModularFeature("SourceControl", &PlasticSourceControlProvider);
+}
+
+void FPlasticSourceControlModule::SaveSettings()
+{
+	if (FApp::IsUnattended() || IsRunningCommandlet())
+	{
+		return;
+	}
+
+	PlasticSourceControlSettings.SaveSettings();
+}
+
+IMPLEMENT_MODULE(FPlasticSourceControlModule, PlasticSourceControl);
+
+#undef LOCTEXT_NAMESPACE
diff --git a/Engine/Plugins/Developer/PlasticSourceControl/Source/PlasticSourceControl/Private/PlasticSourceControlModule.h b/Engine/Plugins/Developer/PlasticSourceControl/Source/PlasticSourceControl/Private/PlasticSourceControlModule.h
new file mode 100644
--- /dev/null
+++ b/Engine/Plugins/Developer/PlasticSourceControl/Source/PlasticSourceControl/Private/PlasticSourceControlModule.h
@@ -0,0 +1,51 @@
+// Copyright 1998-2019 Epic Games, Inc. All Rights Reserved.
+
+#pragma once
+
+#include "CoreMinimal.h"
+#include "Modules/ModuleInterface.h"
+#include "PlasticSourceControlSettings.h"
+#include "PlasticSourceControlProvider.h"
+
+/**
+ * UE4PlasticPlugin is a simple Plastic SCM Source Control Plugin for Unreal Engine 4
+ *
+ * Written and contributed by [REDACTED_PERSON] ([REDACTED_EMAIL])
+ */
+class FPlasticSourceControlModule : public IModuleInterface
+{
+public:
+	/** IModuleInterface implementation */
+	virtual void StartupModule() override;
+	virtual void ShutdownModule() override;
+
+	/** Access the Plastic source control settings */
+	FPlasticSourceControlSettings& AccessSettings()
+	{
+		return PlasticSourceControlSettings;
+	}
+	const FPlasticSourceControlSettings& AccessSettings() const
+	{
+		return PlasticSourceControlSettings;
+	}
+
+	/** Save the Plastic source control settings */
+	void SaveSettings();
+
+	/** Access the Plastic source control provider */
+	FPlasticSourceControlProvider& GetProvider()
+	{
+		return PlasticSourceControlProvider;
+	}
+	const FPlasticSourceControlProvider& GetProvider() const
+	{
+		return PlasticSourceControlProvider;
+	}
+
+private:
+	/** The Plastic source control provider */
+	FPlasticSourceControlProvider PlasticSourceControlProvider;
+
+	/** The settings for Plastic source control */
+	FPlasticSourceControlSettings PlasticSourceControlSettings;
+};
\ No newline at end of file
diff --git a/Engine/Plugins/Developer/PlasticSourceControl/Source/PlasticSourceControl/Private/PlasticSourceControlOperations.cpp b/Engine/Plugins/Developer/PlasticSourceControl/Source/PlasticSourceControl/Private/PlasticSourceControlOperations.cpp
new file mode 100644
--- /dev/null
+++ b/Engine/Plugins/Developer/PlasticSourceControl/Source/PlasticSourceControl/Private/PlasticSourceControlOperations.cpp
@@ -0,0 +1,766 @@
+// Copyright 1998-2019 Epic Games, Inc. All Rights Reserved.
+
+#include "PlasticSourceControlOperations.h"
+#include "PlasticSourceControlSettings.h"
+#include "PlasticSourceControlState.h"
+#include "PlasticSourceControlCommand.h"
+#include "PlasticSourceControlModule.h"
+#include "PlasticSourceControlUtils.h"
+
+#include "AssetRegistryModule.h"
+#include "Async/Async.h"
+#include "HAL/FileManager.h"
+#include "Misc/Paths.h"
+#include "SourceControlOperations.h"
+
+#define LOCTEXT_NAMESPACE "PlasticSourceControl"
+
+FName FPlasticRevertUnchanged::GetName() const
+{
+	return "RevertUnchanged";
+}
+
+FText FPlasticRevertUnchanged::GetInProgressString() const
+{
+	return LOCTEXT("SourceControl_RevertUnchanged", "Reverting unchanged file(s) in Source Control...");
+}
+
+FName FPlasticRevertAll::GetName() const
+{
+	return "RevertAll";
+}
+
+FText FPlasticRevertAll::GetInProgressString() const
+{
+	return LOCTEXT("SourceControl_RevertAll", "Reverting checked-out file(s) in Source Control...");
+}
+
+FName FPlasticMakeWorkspace::GetName() const
+{
+	return "MakeWorkspace";
+}
+
+FText FPlasticMakeWorkspace::GetInProgressString() const
+{
+	return LOCTEXT("SourceControl_MakeWorkspace", "Create a new Repository and initialize the Workspace");
+}
+
+
+FName FPlasticConnectWorker::GetName() const
+{
+	return "Connect";
+}
+
+bool FPlasticConnectWorker::Execute(FPlasticSourceControlCommand& InCommand)
+{
+	check(InCommand.Operation->GetName() == GetName());
+	TSharedRef<FConnect, ESPMode::ThreadSafe> Operation = StaticCastSharedRef<FConnect>(InCommand.Operation);
+
+	const FPlasticSourceControlModule& PlasticSourceControl = FModuleManager::GetModuleChecked<FPlasticSourceControlModule>("PlasticSourceControl");
+	if (PlasticSourceControl.GetProvider().IsPlasticAvailable())
+	{
+		// Get workspace name
+		InCommand.bCommandSuccessful = PlasticSourceControlUtils::GetWorkspaceName(InCommand.WorkspaceName);
+		if (InCommand.bCommandSuccessful)
+		{
+			// Get repository, server Url, branch and current changeset number
+			InCommand.bCommandSuccessful = PlasticSourceControlUtils::GetWorkspaceInformation(InCommand.ChangesetNumber, InCommand.RepositoryName, InCommand.ServerUrl, InCommand.BranchName);
+			if (InCommand.bCommandSuccessful)
+			{
+				// Execute a 'checkconnection' command to check the connectivity of the server.
+				InCommand.bCommandSuccessful = PlasticSourceControlUtils::RunCommand(TEXT("checkconnection"), TArray<FString>(), TArray<FString>(), InCommand.Concurrency, InCommand.InfoMessages, InCommand.ErrorMessages);
+				if (InCommand.bCommandSuccessful)
+				{
+					// Then, update the status of assets in Content/ directory and also Config files,
+					// but only on real (re-)connection (but not each time Login() is called by Rename or Fixup Redirector command to check connection)
+					// and only if enabled in the settings
+					if (!PlasticSourceControl.GetProvider().IsAvailable() && PlasticSourceControl.AccessSettings().UpdateStatusAtStartup())
+					{
+						TArray<FString> ProjectDirs;
+						ProjectDirs.Add(FPaths::ConvertRelativePathToFull(FPaths::ProjectContentDir()));
+						ProjectDirs.Add(FPaths::ConvertRelativePathToFull(FPaths::ProjectConfigDir()));
+						PlasticSourceControlUtils::RunUpdateStatus(ProjectDirs, false, InCommand.Concurrency, InCommand.ErrorMessages, States, InCommand.ChangesetNumber, InCommand.BranchName);
+					}
+				}
+				else
+				{
+					Operation->SetErrorText(FText::FromString(InCommand.ErrorMessages[0]));
+				}
+			}
+		}
+		else
+		{
+			Operation->SetErrorText(LOCTEXT("NotAPlasticRepository", "Failed to enable Plastic SCM source control. You need to initialize the project as a Plastic SCM repository first."));
+		}
+	}
+	else
+	{
+		Operation->SetErrorText(LOCTEXT("PlasticScmCliUnavaillable", "Failed to launch Plastic SCM command line tool. You need to install it or set the correct path to it first."));
+	}
+
+	return InCommand.bCommandSuccessful;
+}
+
+bool FPlasticConnectWorker::UpdateStates()
+{
+	return PlasticSourceControlUtils::UpdateCachedStates(MoveTemp(States));
+}
+
+FName FPlasticCheckOutWorker::GetName() const
+{
+	return "CheckOut";
+}
+
+bool FPlasticCheckOutWorker::Execute(FPlasticSourceControlCommand& InCommand)
+{
+	check(InCommand.Operation->GetName() == GetName());
+
+	// Detect special case for a partial checkout (CS:-1 in Gluon mode)!
+	if (-1 != InCommand.ChangesetNumber)
+	{
+		InCommand.bCommandSuccessful = PlasticSourceControlUtils::RunCommand(TEXT("checkout"), TArray<FString>(), InCommand.Files, InCommand.Concurrency, InCommand.InfoMessages, InCommand.ErrorMessages);
+	}
+	else
+	{
+		InCommand.bCommandSuccessful = PlasticSourceControlUtils::RunCommand(TEXT("partial checkout"), TArray<FString>(), InCommand.Files, InCommand.Concurrency, InCommand.InfoMessages, InCommand.ErrorMessages);
+	}
+
+	// now update the status of our files
+	PlasticSourceControlUtils::RunUpdateStatus(InCommand.Files, false, InCommand.Concurrency, InCommand.ErrorMessages, States, InCommand.ChangesetNumber, InCommand.BranchName);
+
+	return InCommand.bCommandSuccessful;
+}
+
+bool FPlasticCheckOutWorker::UpdateStates()
+{
+	return PlasticSourceControlUtils::UpdateCachedStates(MoveTemp(States));
+}
+
+/// Parse check-in result, usually locking like "Created changeset cs:8@br:/main@MyProject@SRombauts@cloud (mount:'/')"
+static FText ParseCheckInResults(const TArray<FString>& InResults)
+{
+	if (InResults.Num() > 0)
+	{
+		static const FString ChangesetPrefix(TEXT("Created changeset "));
+		if (InResults.Last().StartsWith(ChangesetPrefix))
+		{
+			FString ChangesetString;
+			static const FString BranchPrefix(TEXT("@br:"));
+			const int32 BranchIndex = InResults.Last().Find(BranchPrefix, ESearchCase::CaseSensitive);
+			if (BranchIndex > INDEX_NONE)
+			{
+				ChangesetString = InResults.Last().Mid(ChangesetPrefix.Len(), BranchIndex - ChangesetPrefix.Len());
+			}
+			return FText::Format(LOCTEXT("SubmitMessage", "Submitted changeset {0}"), FText::FromString(ChangesetString));
+		}
+		else
+		{
+			return FText::FromString(InResults.Last());
+		}
+	}
+	return FText();
+}
+
+FName FPlasticCheckInWorker::GetName() const
+{
+	return "CheckIn";
+}
+
+bool FPlasticCheckInWorker::Execute(FPlasticSourceControlCommand& InCommand)
+{
+	check(InCommand.Operation->GetName() == GetName());
+	TSharedRef<FCheckIn, ESPMode::ThreadSafe> Operation = StaticCastSharedRef<FCheckIn>(InCommand.Operation);
+
+	// make a temp file to place our commit message in
+	FScopedTempFile CommitMsgFile(Operation->GetDescription());
+	if (CommitMsgFile.GetFilename().Len() > 0)
+	{
+		TArray<FString> Parameters;
+		FString ParamCommitMsgFilename = TEXT("--commentsfile=\"");
+		ParamCommitMsgFilename += FPaths::ConvertRelativePathToFull(CommitMsgFile.GetFilename());
+		ParamCommitMsgFilename += TEXT("\"");
+		Parameters.Add(ParamCommitMsgFilename);
+		// Detect special case for a partial checkout (CS:-1 in Gluon mode)!
+		if (-1 != InCommand.ChangesetNumber)
+		{
+			Parameters.Add(TEXT("--all")); // Also files Changed (not CheckedOut) and Moved/Deleted Locally
+		//  NOTE: --update added as #23 but removed as #32 because most assets are locked by the Unreal Editor
+		//	Parameters.Add(TEXT("--update")); // Processes the update-merge automatically if it eventually happens.
+			InCommand.bCommandSuccessful = PlasticSourceControlUtils::RunCommand(TEXT("checkin"), Parameters, InCommand.Files, InCommand.Concurrency, InCommand.InfoMessages, InCommand.ErrorMessages);
+		}
+		else
+		{
+			Parameters.Add(TEXT("--applychanged")); // Also files Changed (not CheckedOut) and Moved/Deleted Locally
+			InCommand.bCommandSuccessful = PlasticSourceControlUtils::RunCommand(TEXT("partial checkin"), Parameters, InCommand.Files, InCommand.Concurrency, InCommand.InfoMessages, InCommand.ErrorMessages);
+		}
+		if (InCommand.bCommandSuccessful)
+		{
+			// Remove any deleted files from status cache
+			FPlasticSourceControlModule& PlasticSourceControl = FModuleManager::GetModuleChecked<FPlasticSourceControlModule>("PlasticSourceControl");
+			FPlasticSourceControlProvider& Provider = PlasticSourceControl.GetProvider();
+
+			TArray<TSharedRef<ISourceControlState, ESPMode::ThreadSafe>> LocalStates;
+			Provider.GetState(InCommand.Files, LocalStates, EStateCacheUsage::Use);
+			for (const auto& State : LocalStates)
+			{
+				if (State->IsDeleted())
+				{
+					Provider.RemoveFileFromCache(State->GetFilename());
+				}
+			}
+
+			Operation->SetSuccessMessage(ParseCheckInResults(InCommand.InfoMessages));
+			UE_LOG(LogSourceControl, Log, TEXT("CheckIn successful"));
+		}
+	}
+
+	// now update the status of our files
+	PlasticSourceControlUtils::RunUpdateStatus(InCommand.Files, false, InCommand.Concurrency, InCommand.ErrorMessages, States, InCommand.ChangesetNumber, InCommand.BranchName);
+
+	return InCommand.bCommandSuccessful;
+}
+
+bool FPlasticCheckInWorker::UpdateStates()
+{
+	return PlasticSourceControlUtils::UpdateCachedStates(MoveTemp(States));
+}
+
+FName FPlasticMarkForAddWorker::GetName() const
+{
+	return "MarkForAdd";
+}
+
+bool FPlasticMarkForAddWorker::Execute(FPlasticSourceControlCommand& InCommand)
+{
+	check(InCommand.Operation->GetName() == GetName());
+
+	TArray<FString> Parameters;
+	Parameters.Add(TEXT("--parents"));
+	// Detect special case for a partial checkout (CS:-1 in Gluon mode)!
+	if (-1 != InCommand.ChangesetNumber)
+	{
+		Parameters.Add(TEXT("-R")); // needed at the time of workspace creation, but not working in a partial workspace
+		InCommand.bCommandSuccessful = PlasticSourceControlUtils::RunCommand(TEXT("add"), Parameters, InCommand.Files, InCommand.Concurrency, InCommand.InfoMessages, InCommand.ErrorMessages);
+	}
+	else
+	{
+		InCommand.bCommandSuccessful = PlasticSourceControlUtils::RunCommand(TEXT("partial add"), Parameters, InCommand.Files, InCommand.Concurrency, InCommand.InfoMessages, InCommand.ErrorMessages);
+	}
+
+	// now update the status of our files
+	PlasticSourceControlUtils::RunUpdateStatus(InCommand.Files, false, InCommand.Concurrency, InCommand.ErrorMessages, States, InCommand.ChangesetNumber, InCommand.BranchName);
+
+	return InCommand.bCommandSuccessful;
+}
+
+bool FPlasticMarkForAddWorker::UpdateStates()
+{
+	return PlasticSourceControlUtils::UpdateCachedStates(MoveTemp(States));
+}
+
+FName FPlasticDeleteWorker::GetName() const
+{
+	return "Delete";
+}
+
+bool FPlasticDeleteWorker::Execute(FPlasticSourceControlCommand& InCommand)
+{
+	check(InCommand.Operation->GetName() == GetName());
+
+	// Detect special case for a partial checkout (CS:-1 in Gluon mode)!
+	if (-1 != InCommand.ChangesetNumber)
+	{
+		InCommand.bCommandSuccessful = PlasticSourceControlUtils::RunCommand(TEXT("remove"), TArray<FString>(), InCommand.Files, InCommand.Concurrency, InCommand.InfoMessages, InCommand.ErrorMessages);
+	}
+	else
+	{
+		InCommand.bCommandSuccessful = PlasticSourceControlUtils::RunCommand(TEXT("partial remove"), TArray<FString>(), InCommand.Files, InCommand.Concurrency, InCommand.InfoMessages, InCommand.ErrorMessages);
+	}
+
+	// now update the status of our files
+	PlasticSourceControlUtils::RunUpdateStatus(InCommand.Files, false, InCommand.Concurrency, InCommand.ErrorMessages, States, InCommand.ChangesetNumber, InCommand.BranchName);
+
+	return InCommand.bCommandSuccessful;
+}
+
+bool FPlasticDeleteWorker::UpdateStates()
+{
+	return PlasticSourceControlUtils::UpdateCachedStates(MoveTemp(States));
+}
+
+FName FPlasticRevertWorker::GetName() const
+{
+	return "Revert";
+}
+
+bool FPlasticRevertWorker::Execute(FPlasticSourceControlCommand& InCommand)
+{
+	check(InCommand.Operation->GetName() == GetName());
+
+	FPlasticSourceControlModule& PlasticSourceControl = FModuleManager::GetModuleChecked<FPlasticSourceControlModule>("PlasticSourceControl");
+	FPlasticSourceControlProvider& Provider = PlasticSourceControl.GetProvider();
+
+	InCommand.bCommandSuccessful = true;
+
+	for(const FString& File : InCommand.Files)
+	{
+		TSharedRef<FPlasticSourceControlState, ESPMode::ThreadSafe> State = Provider.GetStateInternal(File);
+
+		TArray<FString> Files;
+		Files.Add(State->LocalFilename);
+
+		if (EWorkspaceState::Changed == State->WorkspaceState)
+		{
+			// revert the changes of the given file in workspace
+			InCommand.bCommandSuccessful &= PlasticSourceControlUtils::RunCommand(TEXT("undochange"), TArray<FString>(), Files, InCommand.Concurrency, InCommand.InfoMessages, InCommand.ErrorMessages);
+		}
+		else
+		{
+			// in case of a Moved/Renamed, find the rename origin to revert both at once
+			if (EWorkspaceState::Moved == State->WorkspaceState)
+			{
+				Files.Add(State->MovedFrom);
+
+				// Delete the redirector
+				IFileManager::Get().Delete(*State->MovedFrom);
+			}
+
+			// revert the checkout and any changes of the given file in workspace
+			// Detect special case for a partial checkout (CS:-1 in Gluon mode)!
+			if (-1 != InCommand.ChangesetNumber)
+			{
+				InCommand.bCommandSuccessful &= PlasticSourceControlUtils::RunCommand(TEXT("undocheckout"), TArray<FString>(), Files, InCommand.Concurrency, InCommand.InfoMessages, InCommand.ErrorMessages);
+			}
+			else
+			{
+				InCommand.bCommandSuccessful &= PlasticSourceControlUtils::RunCommand(TEXT("partial undocheckout"), TArray<FString>(), Files, InCommand.Concurrency, InCommand.InfoMessages, InCommand.ErrorMessages);
+			}
+		}
+	}
+
+	// NOTE: optim, no need to update the status of our files since this is done immediately after by the Editor
+
+	return InCommand.bCommandSuccessful;
+}
+
+bool FPlasticRevertWorker::UpdateStates()
+{
+	return false;
+}
+
+FName FPlasticRevertUnchangedWorker::GetName() const
+{
+	return "RevertUnchanged";
+}
+
+bool FPlasticRevertUnchangedWorker::Execute(FPlasticSourceControlCommand& InCommand)
+{
+	check(InCommand.Operation->GetName() == GetName());
+
+	TArray<FString> Parameters;
+	Parameters.Add(TEXT("-R"));
+
+	// revert the checkout of all unchanged files recursively
+	InCommand.bCommandSuccessful = PlasticSourceControlUtils::RunCommand(TEXT("uncounchanged"), Parameters, InCommand.Files, InCommand.Concurrency, InCommand.InfoMessages, InCommand.ErrorMessages);
+
+	// Now update the status of assets in Content/ directory and also Config files
+	TArray<FString> ProjectDirs;
+	ProjectDirs.Add(FPaths::ConvertRelativePathToFull(FPaths::ProjectContentDir()));
+	ProjectDirs.Add(FPaths::ConvertRelativePathToFull(FPaths::ProjectConfigDir()));
+	PlasticSourceControlUtils::RunUpdateStatus(ProjectDirs, false, InCommand.Concurrency, InCommand.ErrorMessages, States, InCommand.ChangesetNumber, InCommand.BranchName);
+
+	return InCommand.bCommandSuccessful;
+}
+
+bool FPlasticRevertUnchangedWorker::UpdateStates()
+{
+	return PlasticSourceControlUtils::UpdateCachedStates(MoveTemp(States));
+}
+
+FName FPlasticRevertAllWorker::GetName() const
+{
+	return "RevertAll";
+}
+
+bool FPlasticRevertAllWorker::Execute(FPlasticSourceControlCommand& InCommand)
+{
+	check(InCommand.Operation->GetName() == GetName());
+
+	TArray<FString> Parameters;
+	Parameters.Add(TEXT("--all"));
+	// revert the checkout of all files recursively
+	// Detect special case for a partial checkout (CS:-1 in Gluon mode)!
+	if (-1 != InCommand.ChangesetNumber)
+	{
+		InCommand.bCommandSuccessful = PlasticSourceControlUtils::RunCommand(TEXT("undocheckout"), Parameters, InCommand.Files, InCommand.Concurrency, InCommand.InfoMessages, InCommand.ErrorMessages);
+	}
+	else
+	{
+		InCommand.bCommandSuccessful = PlasticSourceControlUtils::RunCommand(TEXT("partial undocheckout"), Parameters, InCommand.Files, InCommand.Concurrency, InCommand.InfoMessages, InCommand.ErrorMessages);
+	}
+
+	// Now update the status of assets in Content/ directory and also Config files
+	TArray<FString> ProjectDirs;
+	ProjectDirs.Add(FPaths::ConvertRelativePathToFull(FPaths::ProjectContentDir()));
+	ProjectDirs.Add(FPaths::ConvertRelativePathToFull(FPaths::ProjectConfigDir()));
+	PlasticSourceControlUtils::RunUpdateStatus(ProjectDirs, false, InCommand.Concurrency, InCommand.ErrorMessages, States, InCommand.ChangesetNumber, InCommand.BranchName);
+
+	return InCommand.bCommandSuccessful;
+}
+
+bool FPlasticRevertAllWorker::UpdateStates()
+{
+	return PlasticSourceControlUtils::UpdateCachedStates(MoveTemp(States));
+}
+
+FName FPlasticMakeWorkspaceWorker::GetName() const
+{
+	return "MakeWorkspace";
+}
+
+bool FPlasticMakeWorkspaceWorker::Execute(FPlasticSourceControlCommand& InCommand)
+{
+	check(InCommand.Operation->GetName() == GetName());
+	TSharedRef<FPlasticMakeWorkspace, ESPMode::ThreadSafe> Operation = StaticCastSharedRef<FPlasticMakeWorkspace>(InCommand.Operation);
+
+	{
+		TArray<FString> Parameters;
+		Parameters.Add(Operation->ServerUrl);
+		Parameters.Add(Operation->RepositoryName);
+		PlasticSourceControlUtils::RunCommand(TEXT("makerepository"), Parameters, TArray<FString>(), EConcurrency::Synchronous, InCommand.InfoMessages, InCommand.ErrorMessages);
+	}
+	{
+		TArray<FString> Parameters;
+		Parameters.Add(Operation->WorkspaceName);
+		Parameters.Add(TEXT(".")); // current path, ie. ProjectDir
+		Parameters.Add(FString::Printf(TEXT("--repository=rep:%s@repserver:%s"), *Operation->RepositoryName, *Operation->ServerUrl));
+		InCommand.bCommandSuccessful = PlasticSourceControlUtils::RunCommand(TEXT("makeworkspace"), Parameters, TArray<FString>(), EConcurrency::Synchronous, InCommand.InfoMessages, InCommand.ErrorMessages);
+	}
+
+	return InCommand.bCommandSuccessful;
+}
+
+bool FPlasticMakeWorkspaceWorker::UpdateStates()
+{
+	return false;
+}
+
+FName FPlasticUpdateStatusWorker::GetName() const
+{
+	return "UpdateStatus";
+}
+
+bool FPlasticUpdateStatusWorker::Execute(FPlasticSourceControlCommand& InCommand)
+{
+	check(InCommand.Operation->GetName() == GetName());
+	TSharedRef<FUpdateStatus, ESPMode::ThreadSafe> Operation = StaticCastSharedRef<FUpdateStatus>(InCommand.Operation);
+
+	// @todo: temporary debug log
+	UE_LOG(LogSourceControl, Log, TEXT("status (of %d files, ShouldCheckAllFiles=%d, ShouldUpdateHistory=%d, ShouldGetOpenedOnly=%d, ShouldUpdateModifiedState=%d)"),
+		InCommand.Files.Num(), Operation->ShouldCheckAllFiles()?1:0, Operation->ShouldUpdateHistory()?1:0, Operation->ShouldGetOpenedOnly()?1:0, Operation->ShouldUpdateModifiedState()?1:0);
+
+	if (InCommand.Files.Num() > 0)
+	{
+		InCommand.bCommandSuccessful = PlasticSourceControlUtils::RunUpdateStatus(InCommand.Files, Operation->ShouldUpdateHistory(), InCommand.Concurrency, InCommand.ErrorMessages, States, InCommand.ChangesetNumber, InCommand.BranchName);
+		// Remove all "is not in a workspace" error and convert the result to "success" if there are no other errors
+		PlasticSourceControlUtils::RemoveRedundantErrors(InCommand, TEXT("is not in a workspace."));
+		if (!InCommand.bCommandSuccessful)
+		{
+			UE_LOG(LogSourceControl, Error, TEXT("FPlasticUpdateStatusWorker(ErrorMessages.Num()=%d) => checkconnection"), InCommand.ErrorMessages.Num());
+			// In case of error, execute a 'checkconnection' command to check the connectivity of the server.
+			InCommand.bConnectionDropped = !PlasticSourceControlUtils::RunCommand(TEXT("checkconnection"), TArray<FString>(), TArray<FString>(), InCommand.Concurrency, InCommand.InfoMessages, InCommand.ErrorMessages);
+		}
+		else
+		{
+			if (Operation->ShouldUpdateHistory())
+			{
+				for (int32 IdxFile = 0; IdxFile < States.Num(); IdxFile++)
+				{
+					FString& File = InCommand.Files[IdxFile];
+					FPlasticSourceControlState& State = States[IdxFile];
+
+					if (State.IsSourceControlled())
+					{
+						// Get the history of the file (on all branches)
+						InCommand.bCommandSuccessful &= PlasticSourceControlUtils::RunGetHistory(File, InCommand.ErrorMessages, State);
+						if (State.IsConflicted())
+						{
+							// In case of a merge conflict, we need to put the tip of the "remote branch" on top of the history
+							UE_LOG(LogSourceControl, Log, TEXT("%s: PendingMergeSourceChangeset %d"), *State.LocalFilename, State.PendingMergeSourceChangeset);
+							for (int32 IdxRevision = 0; IdxRevision < State.History.Num(); IdxRevision++)
+							{
+								const auto& Revision = State.History[IdxRevision];
+								if (Revision->ChangesetNumber == State.PendingMergeSourceChangeset)
+								{
+									// If the Source Changeset is not already at the top of the History, duplicate it there.
+									if (IdxRevision > 0)
+									{
+										const auto RevisionCopy = Revision;
+										State.History.Insert(RevisionCopy, 0);
+									}
+									break;
+								}
+							}
+						}
+					}
+				}
+			}
+		}
+	}
+	else
+	{
+		// no path provided: only update the status of assets in Content/ directory and also Config files
+		TArray<FString> ProjectDirs;
+		ProjectDirs.Add(FPaths::ConvertRelativePathToFull(FPaths::ProjectContentDir()));
+		ProjectDirs.Add(FPaths::ConvertRelativePathToFull(FPaths::ProjectConfigDir()));
+		InCommand.bCommandSuccessful = PlasticSourceControlUtils::RunUpdateStatus(ProjectDirs, false, InCommand.Concurrency, InCommand.ErrorMessages, States, InCommand.ChangesetNumber, InCommand.BranchName);
+	}
+
+	// don't use the ShouldUpdateModifiedState() hint here as it is specific to Perforce: the above normal Plastic status has already told us this information (like Git and Mercurial)
+
+	return InCommand.bCommandSuccessful;
+}
+
+bool FPlasticUpdateStatusWorker::UpdateStates()
+{
+	return PlasticSourceControlUtils::UpdateCachedStates(MoveTemp(States));
+}
+
+/// Detect if the operation is a duplicate/copy or a rename/move, and if it leaved a redirector (ie it was a move of a source controled asset)
+bool IsMoveOperation(const FString& InOrigin)
+{
+	bool bIsMoveOperation = true;
+
+	FString PackageName;
+	if (FPackageName::TryConvertFilenameToLongPackageName(InOrigin, PackageName))
+	{
+		// Use AsyncTask to call AssetRegistry GetAssetsByPackageName') on Game Thread
+		const TSharedRef<TPromise<TArray<FAssetData>>, ESPMode::ThreadSafe> Promise = MakeShareable(new TPromise<TArray<FAssetData>>());
+		AsyncTask(ENamedThreads::GameThread, [Promise, PackageName]()
+		{
+			TArray<FAssetData> AssetsData;
+			FAssetRegistryModule& AssetRegistryModule = FModuleManager::GetModuleChecked<FAssetRegistryModule>(TEXT("AssetRegistry"));
+			AssetRegistryModule.Get().GetAssetsByPackageName(FName(*PackageName), AssetsData);
+			Promise->SetValue(MoveTemp(AssetsData));
+		});
+		const TArray<FAssetData> AssetsData = Promise->GetFuture().Get();
+		UE_LOG(LogSourceControl, Log, TEXT("PackageName: %s, AssetsData: Num=%d"), *PackageName, AssetsData.Num());
+		if (AssetsData.Num() > 0)
+		{
+			const FAssetData& AssetData = AssetsData[0];
+			if (!AssetData.IsRedirector())
+			{
+				UE_LOG(LogSourceControl, Log, TEXT("%s is a plain asset, so it's a duplicate/copy"), *InOrigin);
+				bIsMoveOperation = false;
+			}
+			else
+			{
+				UE_LOG(LogSourceControl, Log, TEXT("%s is a redirector, so it's a move/rename"), *InOrigin);
+			}
+		}
+		else
+		{
+			// no asset in package (no redirector) so it should be a rename/move of a newly Added (not Controlled/Checked-In) file
+			UE_LOG(LogSourceControl, Log, TEXT("%s does not have asset in package (ie. no redirector) so it's a move/rename of a newly added file"), *InOrigin);
+		}
+	}
+
+	return bIsMoveOperation;
+}
+
+FName FPlasticCopyWorker::GetName() const
+{
+	return "Copy";
+}
+
+bool FPlasticCopyWorker::Execute(FPlasticSourceControlCommand& InCommand)
+{
+	check(InCommand.Operation->GetName() == GetName());
+	TSharedRef<FCopy, ESPMode::ThreadSafe> Operation = StaticCastSharedRef<FCopy>(InCommand.Operation);
+
+	if (InCommand.Files.Num() == 1)
+	{
+		const FString& Origin = InCommand.Files[0];
+		const FString Destination = FPaths::ConvertRelativePathToFull(Operation->GetDestination());
+
+		// Detect if the operation is a duplicate/copy or a rename/move, and if it leaved a redirector (ie it was a move of a source controled asset)
+		const bool bIsMoveOperation = IsMoveOperation(Origin);
+		if (bIsMoveOperation)
+		{
+			UE_LOG(LogSourceControl, Log, TEXT("Moving %s to %s..."), *Origin, *Destination);
+			// In case of rename, we have to undo what the Editor (created a redirector and added the dest asset), and then redo it with Plastic SCM
+			// - backup the redirector (if it exists) to a temp file
+			const bool bReplace = true;
+			const bool bEvenIfReadOnly = true;
+			const FString TempFileName = FPaths::CreateTempFilename(*FPaths::ProjectLogDir(), TEXT("Plastic-MoveTemp"), TEXT(".uasset"));
+			UE_LOG(LogSourceControl, Log, TEXT("Move '%s' -> '%d'"), *Origin, *TempFileName);
+			InCommand.bCommandSuccessful = IFileManager::Get().Move(*TempFileName, *Origin, bReplace, bEvenIfReadOnly);
+			// - revert the 'cm add' that was applied to the destination by the Editor
+			if (InCommand.bCommandSuccessful)
+			{
+				TArray<FString> DestinationFiles;
+				DestinationFiles.Add(Destination);
+				InCommand.bCommandSuccessful = PlasticSourceControlUtils::RunCommand(TEXT("undochange"), TArray<FString>(), DestinationFiles, InCommand.Concurrency, InCommand.InfoMessages, InCommand.ErrorMessages);
+			}
+			// - move back the asset from the destination to it's original location
+			if (InCommand.bCommandSuccessful)
+			{
+				UE_LOG(LogSourceControl, Log, TEXT("Move '%s' -> '%d'"), *Destination, *Origin);
+				InCommand.bCommandSuccessful = IFileManager::Get().Move(*Origin, *Destination, bReplace, bEvenIfReadOnly);
+			}
+			// - execute a 'cm move' command to the destination to redo the actual job
+			if (InCommand.bCommandSuccessful)
+			{
+				TArray<FString> Files;
+				Files.Add(Origin);
+				Files.Add(Destination);
+				// Detect special case for a partial checkout (CS:-1 in Gluon mode)!
+				if (-1 != InCommand.ChangesetNumber)
+				{
+					InCommand.bCommandSuccessful = PlasticSourceControlUtils::RunCommand(TEXT("move"), TArray<FString>(), Files, InCommand.Concurrency, InCommand.InfoMessages, InCommand.ErrorMessages);
+				}
+				else
+				{
+					InCommand.bCommandSuccessful = PlasticSourceControlUtils::RunCommand(TEXT("partial move"), TArray<FString>(), Files, InCommand.Concurrency, InCommand.InfoMessages, InCommand.ErrorMessages);
+				}
+			}
+			// - restore the redirector file (if it exists) to it's former location
+			if (InCommand.bCommandSuccessful)
+			{
+				UE_LOG(LogSourceControl, Log, TEXT("Move '%s' -> '%d'"), *TempFileName, *Origin);
+				InCommand.bCommandSuccessful = IFileManager::Get().Move(*Origin, *TempFileName, bReplace, bEvenIfReadOnly);
+			}
+			// - add the redirector file (if it exists) to source control
+			if (InCommand.bCommandSuccessful)
+			{
+				TArray<FString> Files;
+				Files.Add(Origin);
+				// Detect special case for a partial checkout (CS:-1 in Gluon mode)!
+				if (-1 != InCommand.ChangesetNumber)
+				{
+					InCommand.bCommandSuccessful = PlasticSourceControlUtils::RunCommand(TEXT("add"), TArray<FString>(), Files, InCommand.Concurrency, InCommand.InfoMessages, InCommand.ErrorMessages);
+				}
+				else
+				{
+					InCommand.bCommandSuccessful = PlasticSourceControlUtils::RunCommand(TEXT("partial add"), TArray<FString>(), Files, InCommand.Concurrency, InCommand.InfoMessages, InCommand.ErrorMessages);
+				}
+			}
+		}
+		else
+		{
+			// copy operation: destination file already added to Source Control, and original asset not changed, so nothing to do
+			InCommand.bCommandSuccessful = true;
+		}
+
+		// now update the status of our files:
+		TArray<FString> BothFiles;
+		BothFiles.Add(Origin);
+		BothFiles.Add(Destination);
+		PlasticSourceControlUtils::RunUpdateStatus(BothFiles, false, InCommand.Concurrency, InCommand.ErrorMessages, States, InCommand.ChangesetNumber, InCommand.BranchName);
+	}
+	else
+	{
+		UE_LOG(LogSourceControl, Error, TEXT("Copy is working for one file only: %d provided!"), InCommand.Files.Num());
+	}
+
+	return InCommand.bCommandSuccessful;
+}
+
+bool FPlasticCopyWorker::UpdateStates()
+{
+	return PlasticSourceControlUtils::UpdateCachedStates(MoveTemp(States));
+}
+
+FName FPlasticSyncWorker::GetName() const
+{
+	return "Sync";
+}
+
+bool FPlasticSyncWorker::Execute(FPlasticSourceControlCommand& InCommand)
+{
+	check(InCommand.Operation->GetName() == GetName());
+
+	TArray<FString> Parameters;
+	// Update specified directory to the head of the repository
+	// Detect special case for a partial checkout (CS:-1 in Gluon mode)!
+	if (-1 != InCommand.ChangesetNumber)
+	{
+		Parameters.Add(TEXT("--last"));
+		Parameters.Add(TEXT("--dontmerge"));
+		InCommand.bCommandSuccessful = PlasticSourceControlUtils::RunCommand(TEXT("update"), Parameters, InCommand.Files, InCommand.Concurrency, InCommand.InfoMessages, InCommand.ErrorMessages);
+	}
+	else
+	{
+		InCommand.bCommandSuccessful = PlasticSourceControlUtils::RunCommand(TEXT("partial update"), Parameters, InCommand.Files, InCommand.Concurrency, InCommand.InfoMessages, InCommand.ErrorMessages);
+	}
+
+	if (InCommand.bCommandSuccessful)
+	{
+		// now update the status of our files
+		// detect the special case of a Sync of the root folder:
+		if ((InCommand.Files.Num() == 1) && (InCommand.Files.Last() == InCommand.PathToWorkspaceRoot))
+		{
+			// only update the status of assets in Content/ directory and also Config files
+			TArray<FString> ProjectDirs;
+			ProjectDirs.Add(FPaths::ConvertRelativePathToFull(FPaths::ProjectContentDir()));
+			ProjectDirs.Add(FPaths::ConvertRelativePathToFull(FPaths::ProjectConfigDir()));
+			PlasticSourceControlUtils::RunUpdateStatus(ProjectDirs, false, InCommand.Concurrency, InCommand.ErrorMessages, States, InCommand.ChangesetNumber, InCommand.BranchName);
+		}
+		// else: optim, no need to update the status of our files since this is done immediately after by the Editor
+	}
+
+	return InCommand.bCommandSuccessful;
+}
+
+bool FPlasticSyncWorker::UpdateStates()
+{
+	return PlasticSourceControlUtils::UpdateCachedStates(MoveTemp(States));
+}
+
+FName FPlasticResolveWorker::GetName() const
+{
+	return "Resolve";
+}
+
+bool FPlasticResolveWorker::Execute(FPlasticSourceControlCommand& InCommand)
+{
+	check(InCommand.Operation->GetName() == GetName());
+
+	// Currently resolve operation is always on one file only, but the following would works for many
+	for (const FString& File : InCommand.Files)
+	{
+		FPlasticSourceControlModule& PlasticSourceControl = FModuleManager::GetModuleChecked<FPlasticSourceControlModule>("PlasticSourceControl");
+		FPlasticSourceControlProvider& Provider = PlasticSourceControl.GetProvider();
+		TSharedRef<FPlasticSourceControlState, ESPMode::ThreadSafe> State = Provider.GetStateInternal(File);
+
+		// To resolve the conflict, merge the file by keeping it like it is on file system
+		// TODO: according to documentation, this cannot work for cherry-picking
+		// merge cs:2@repo@url:port --merge --keepdestination "/path/to/file"
+
+		// Use Merge Parameters obtained in the UpdateStatus operation
+		TArray<FString> Parameters = State->PendingMergeParameters;
+		Parameters.Add(TEXT("--merge"));
+		Parameters.Add(TEXT("--keepdestination"));
+
+		TArray<FString> OneFile;
+		OneFile.Add(State->PendingMergeFilename);
+
+		// @todo temporary debug log
+		UE_LOG(LogSourceControl, Log, TEXT("resolve %s"), *State->PendingMergeFilename);
+
+		// Mark the conflicted file as resolved
+		InCommand.bCommandSuccessful = PlasticSourceControlUtils::RunCommand(TEXT("merge"), Parameters, OneFile, InCommand.Concurrency, InCommand.InfoMessages, InCommand.ErrorMessages);
+	}
+
+	// now update the status of our files
+	PlasticSourceControlUtils::RunUpdateStatus(InCommand.Files, false, InCommand.Concurrency, InCommand.ErrorMessages, States, InCommand.ChangesetNumber, InCommand.BranchName);
+
+	return InCommand.bCommandSuccessful;
+}
+
+bool FPlasticResolveWorker::UpdateStates()
+{
+	return PlasticSourceControlUtils::UpdateCachedStates(MoveTemp(States));
+}
+
+#undef LOCTEXT_NAMESPACE
diff --git a/Engine/Plugins/Developer/PlasticSourceControl/Source/PlasticSourceControl/Private/PlasticSourceControlOperations.h b/Engine/Plugins/Developer/PlasticSourceControl/Source/PlasticSourceControl/Private/PlasticSourceControlOperations.h
new file mode 100644
--- /dev/null
+++ b/Engine/Plugins/Developer/PlasticSourceControl/Source/PlasticSourceControl/Private/PlasticSourceControlOperations.h
@@ -0,0 +1,242 @@
+// Copyright 1998-2019 Epic Games, Inc. All Rights Reserved.
+
+// Specialization of classes defineds in Engine\Source\Developer\SourceControl\Public\SourceControlOperations.h
+
+#pragma once
+
+#include "CoreMinimal.h"
+#include "IPlasticSourceControlWorker.h"
+#include "PlasticSourceControlState.h"
+#include "PlasticSourceControlRevision.h"
+
+#include "ISourceControlOperation.h"
+
+
+/**
+ * Internal operation used to revert checked-out unchanged files
+*/
+class FPlasticRevertUnchanged : public ISourceControlOperation
+{
+public:
+	// ISourceControlOperation interface
+	virtual FName GetName() const override;
+
+	virtual FText GetInProgressString() const override;
+};
+
+
+/**
+ * Internal operation used to revert checked-out files
+*/
+class FPlasticRevertAll : public ISourceControlOperation
+{
+public:
+	// ISourceControlOperation interface
+	virtual FName GetName() const override;
+
+	virtual FText GetInProgressString() const override;
+};
+
+
+/**
+* Internal operation used to initialize a new Workspace and a new Repository
+*/
+class FPlasticMakeWorkspace : public ISourceControlOperation
+{
+public:
+	// ISourceControlOperation interface
+	virtual FName GetName() const override;
+
+	virtual FText GetInProgressString() const override;
+
+	FString WorkspaceName;
+	FString RepositoryName;
+	FString ServerUrl;
+};
+
+
+/** Called when first activated on a project, and then at project load time.
+ *  Look for the root directory of the Plastic workspace (where the ".plastic/" subdirectory is located). */
+class FPlasticConnectWorker : public IPlasticSourceControlWorker
+{
+public:
+	virtual ~FPlasticConnectWorker() {}
+	// IPlasticSourceControlWorker interface
+	virtual FName GetName() const override;
+	virtual bool Execute(class FPlasticSourceControlCommand& InCommand) override;
+	virtual bool UpdateStates() override;
+
+public:
+	/** Temporary states for results */
+	TArray<FPlasticSourceControlState> States;
+};
+
+class FPlasticCheckOutWorker : public IPlasticSourceControlWorker
+{
+public:
+	virtual ~FPlasticCheckOutWorker() {}
+	// IPlasticSourceControlWorker interface
+	virtual FName GetName() const override;
+	virtual bool Execute(class FPlasticSourceControlCommand& InCommand) override;
+	virtual bool UpdateStates() override;
+
+public:
+	/** Temporary states for results */
+	TArray<FPlasticSourceControlState> States;
+};
+
+/** Check-in a set of file to the local depot. */
+class FPlasticCheckInWorker : public IPlasticSourceControlWorker
+{
+public:
+	virtual ~FPlasticCheckInWorker() {}
+	// IPlasticSourceControlWorker interface
+	virtual FName GetName() const override;
+	virtual bool Execute(class FPlasticSourceControlCommand& InCommand) override;
+	virtual bool UpdateStates() override;
+
+public:
+	/** Temporary states for results */
+	TArray<FPlasticSourceControlState> States;
+};
+
+/** Add an untracked file to source control (so only a subset of the Plastic add command). */
+class FPlasticMarkForAddWorker : public IPlasticSourceControlWorker
+{
+public:
+	virtual ~FPlasticMarkForAddWorker() {}
+	// IPlasticSourceControlWorker interface
+	virtual FName GetName() const override;
+	virtual bool Execute(class FPlasticSourceControlCommand& InCommand) override;
+	virtual bool UpdateStates() override;
+
+public:
+	/** Temporary states for results */
+	TArray<FPlasticSourceControlState> States;
+};
+
+/** Delete a file and remove it from source control. */
+class FPlasticDeleteWorker : public IPlasticSourceControlWorker
+{
+public:
+	virtual ~FPlasticDeleteWorker() {}
+	// IPlasticSourceControlWorker interface
+	virtual FName GetName() const override;
+	virtual bool Execute(class FPlasticSourceControlCommand& InCommand) override;
+	virtual bool UpdateStates() override;
+
+public:
+	/** Temporary states for results */
+	TArray<FPlasticSourceControlState> States;
+};
+
+/** Revert any change to a file to its state on the local depot. */
+class FPlasticRevertWorker : public IPlasticSourceControlWorker
+{
+public:
+	virtual ~FPlasticRevertWorker() {}
+	// IPlasticSourceControlWorker interface
+	virtual FName GetName() const override;
+	virtual bool Execute(class FPlasticSourceControlCommand& InCommand) override;
+	virtual bool UpdateStates() override;
+};
+
+/** Revert only unchanged file(s) (uncheckout). */
+class FPlasticRevertUnchangedWorker : public IPlasticSourceControlWorker
+{
+public:
+	virtual ~FPlasticRevertUnchangedWorker() {}
+	// IPlasticSourceControlWorker interface
+	virtual FName GetName() const override;
+	virtual bool Execute(class FPlasticSourceControlCommand& InCommand) override;
+	virtual bool UpdateStates() override;
+
+public:
+	/** Temporary states for results */
+	TArray<FPlasticSourceControlState> States;
+};
+
+/** Revert all checked-out file(s). */
+class FPlasticRevertAllWorker : public IPlasticSourceControlWorker
+{
+public:
+	virtual ~FPlasticRevertAllWorker() {}
+	// IPlasticSourceControlWorker interface
+	virtual FName GetName() const override;
+	virtual bool Execute(class FPlasticSourceControlCommand& InCommand) override;
+	virtual bool UpdateStates() override;
+
+public:
+	/** Temporary states for results */
+	TArray<FPlasticSourceControlState> States;
+};
+
+/** Initialize a new Workspace and a new Repository */
+class FPlasticMakeWorkspaceWorker : public IPlasticSourceControlWorker
+{
+public:
+	virtual ~FPlasticMakeWorkspaceWorker() {}
+	// IPlasticSourceControlWorker interface
+	virtual FName GetName() const override;
+	virtual bool Execute(class FPlasticSourceControlCommand& InCommand) override;
+	virtual bool UpdateStates() override;
+};
+
+/** Plastic update the workspace to latest changes */
+class FPlasticSyncWorker : public IPlasticSourceControlWorker
+{
+public:
+	virtual ~FPlasticSyncWorker() {}
+	// IPlasticSourceControlWorker interface
+	virtual FName GetName() const override;
+	virtual bool Execute(class FPlasticSourceControlCommand& InCommand) override;
+	virtual bool UpdateStates() override;
+
+public:
+	/** Temporary states for results */
+	TArray<FPlasticSourceControlState> States;
+};
+
+/** Get source control status of files on local workspace. */
+class FPlasticUpdateStatusWorker : public IPlasticSourceControlWorker
+{
+public:
+	virtual ~FPlasticUpdateStatusWorker() {}
+	// IPlasticSourceControlWorker interface
+	virtual FName GetName() const override;
+	virtual bool Execute(class FPlasticSourceControlCommand& InCommand) override;
+	virtual bool UpdateStates() override;
+
+public:
+	/** Temporary states for results */
+	TArray<FPlasticSourceControlState> States;
+};
+
+/** Copy or Move operation on a single file */
+class FPlasticCopyWorker : public IPlasticSourceControlWorker
+{
+public:
+	virtual ~FPlasticCopyWorker() {}
+	// IPlasticSourceControlWorker interface
+	virtual FName GetName() const override;
+	virtual bool Execute(class FPlasticSourceControlCommand& InCommand) override;
+	virtual bool UpdateStates() override;
+
+public:
+	/** Temporary states for results */
+	TArray<FPlasticSourceControlState> States;
+};
+
+/** Plastic command to mark the conflict as solved */
+class FPlasticResolveWorker : public IPlasticSourceControlWorker
+{
+public:
+	virtual ~FPlasticResolveWorker() {}
+	virtual FName GetName() const override;
+	virtual bool Execute(class FPlasticSourceControlCommand& InCommand) override;
+	virtual bool UpdateStates() override;
+
+private:
+	/** Temporary states for results */
+	TArray<FPlasticSourceControlState> States;
+};
diff --git a/Engine/Plugins/Developer/PlasticSourceControl/Source/PlasticSourceControl/Private/PlasticSourceControlProvider.cpp b/Engine/Plugins/Developer/PlasticSourceControl/Source/PlasticSourceControl/Private/PlasticSourceControlProvider.cpp
new file mode 100644
--- /dev/null
+++ b/Engine/Plugins/Developer/PlasticSourceControl/Source/PlasticSourceControl/Private/PlasticSourceControlProvider.cpp
@@ -0,0 +1,463 @@
+// Copyright 1998-2019 Epic Games, Inc. All Rights Reserved.
+
+#include "PlasticSourceControlProvider.h"
+#include "PlasticSourceControlCommand.h"
+#include "ISourceControlModule.h"
+#include "PlasticSourceControlModule.h"
+#include "PlasticSourceControlSettings.h"
+#include "PlasticSourceControlOperations.h"
+#include "PlasticSourceControlUtils.h"
+#include "SPlasticSourceControlSettings.h"
+#include "Logging/MessageLog.h"
+#include "ScopedSourceControlProgress.h"
+#include "SourceControlHelpers.h"
+#include "SourceControlOperations.h"
+#include "Interfaces/IPluginManager.h"
+
+#include "Misc/Paths.h"
+#include "HAL/PlatformProcess.h"
+#include "Misc/QueuedThreadPool.h"
+
+#define LOCTEXT_NAMESPACE "PlasticSourceControl"
+
+static FName ProviderName("Plastic SCM");
+
+void FPlasticSourceControlProvider::Init(bool bForceConnection)
+{
+	// Init() is called multiple times at startup: do not check Plastic SCM each time
+	if(!bPlasticAvailable)
+	{
+		const TSharedPtr<IPlugin> Plugin = IPluginManager::Get().FindPlugin(TEXT("PlasticSourceControl"));
+		if(Plugin.IsValid())
+		{
+			UE_LOG(LogSourceControl, Log, TEXT("Plastic SCM plugin '%s'"), *(Plugin->GetDescriptor().VersionName));
+		}
+
+		CheckPlasticAvailability();
+	}
+
+	// bForceConnection: not used anymore
+}
+
+void FPlasticSourceControlProvider::CheckPlasticAvailability()
+{
+	FPlasticSourceControlModule& PlasticSourceControl = FModuleManager::LoadModuleChecked<FPlasticSourceControlModule>("PlasticSourceControl");
+	FString PathToPlasticBinary = PlasticSourceControl.AccessSettings().GetBinaryPath();
+	if(PathToPlasticBinary.IsEmpty())
+	{
+		// Try to find Plastic binary, and update settings accordingly
+		PathToPlasticBinary = PlasticSourceControlUtils::FindPlasticBinaryPath();
+		if(!PathToPlasticBinary.IsEmpty())
+		{
+			PlasticSourceControl.AccessSettings().SetBinaryPath(PathToPlasticBinary);
+		}
+	}
+
+	if(!PathToPlasticBinary.IsEmpty())
+	{
+		// Find the path to the root Plastic directory (if any, else uses the ProjectDir)
+		const FString PathToProjectDir = FPaths::ConvertRelativePathToFull(FPaths::ProjectDir());
+		bWorkspaceFound = PlasticSourceControlUtils::FindRootDirectory(PathToProjectDir, PathToWorkspaceRoot);
+
+		// Launch the Plastic SCM cli shell on the background to issue all commands during this session
+		bPlasticAvailable = PlasticSourceControlUtils::LaunchBackgroundPlasticShell(PathToPlasticBinary, PathToWorkspaceRoot);
+		if(bPlasticAvailable)
+		{
+			PlasticSourceControlUtils::GetPlasticScmVersion(PlasticScmVersion);
+
+			// Get user name (from the global Plastic SCM client config)
+			PlasticSourceControlUtils::GetUserName(UserName);
+
+			if(!bWorkspaceFound)
+			{
+				UE_LOG(LogSourceControl, Warning, TEXT("'%s' is not part of a Plastic workspace"), *FPaths::ProjectDir());
+			}
+		}
+	}
+	else
+	{
+		bPlasticAvailable = false;
+	}
+}
+
+void FPlasticSourceControlProvider::Close()
+{
+	// clear the cache
+	StateCache.Empty();
+	// terminate the background 'cm shell' process and associated pipes
+	PlasticSourceControlUtils::Terminate();
+	// Remove all extensions to the "Source Control" menu in the Editor Toolbar
+	PlasticSourceControlMenu.Unregister();
+
+	bServerAvailable = false;
+	bPlasticAvailable = false;
+	bWorkspaceFound = false;
+	UserName.Empty();
+}
+
+TSharedRef<FPlasticSourceControlState, ESPMode::ThreadSafe> FPlasticSourceControlProvider::GetStateInternal(const FString& Filename) const
+{
+	TSharedRef<FPlasticSourceControlState, ESPMode::ThreadSafe>* State = StateCache.Find(Filename);
+	if(State != NULL)
+	{
+		// found cached item
+		return (*State);
+	}
+	else
+	{
+		// cache an unknown state for this item
+		TSharedRef<FPlasticSourceControlState, ESPMode::ThreadSafe> NewState = MakeShareable( new FPlasticSourceControlState(Filename) );
+		StateCache.Add(Filename, NewState);
+		return NewState;
+	}
+}
+
+FText FPlasticSourceControlProvider::GetStatusText() const
+{
+	FFormatNamedArguments Args;
+    Args.Add( TEXT("PlasticScmVersion"), FText::FromString(PlasticScmVersion) );
+	Args.Add( TEXT("WorkspacePath"), FText::FromString(PathToWorkspaceRoot) );
+	Args.Add( TEXT("WorkspaceName"), FText::FromString(WorkspaceName) );
+	Args.Add( TEXT("BranchName"), FText::FromString(BranchName) );
+	// Detect special case for a partial checkout (CS:-1 in Gluon mode)!
+	if (-1 != ChangesetNumber)
+	{
+		Args.Add( TEXT("ChangesetNumber"), FText::FromString(FString::Printf(TEXT("%d (standard mode)"), ChangesetNumber)) );
+	}
+	else
+	{
+		Args.Add( TEXT("ChangesetNumber"), FText::FromString(FString::Printf(TEXT("N/A (Gluon/partial mode)"))) );
+	}
+	Args.Add( TEXT("UserName"), FText::FromString(UserName) );
+
+	return FText::Format( NSLOCTEXT("Status", "Provider: Plastic\nEnabledLabel", "Plastic SCM {PlasticScmVersion}\nWorkspace: {WorkspaceName} ({WorkspacePath})\n{BranchName}\nChangeset: {ChangesetNumber}\nUser: {UserName}"), Args );
+}
+
+/** Quick check if source control is enabled */
+bool FPlasticSourceControlProvider::IsEnabled() const
+{
+	return bWorkspaceFound;
+}
+
+/** Quick check if source control is available for use (useful for server-based providers) */
+bool FPlasticSourceControlProvider::IsAvailable() const
+{
+	return bServerAvailable;
+}
+
+const FName& FPlasticSourceControlProvider::GetName(void) const
+{
+	return ProviderName;
+}
+
+ECommandResult::Type FPlasticSourceControlProvider::GetState( const TArray<FString>& InFiles, TArray< TSharedRef<ISourceControlState, ESPMode::ThreadSafe> >& OutState, EStateCacheUsage::Type InStateCacheUsage )
+{
+	if(!IsEnabled())
+	{
+		return ECommandResult::Failed;
+	}
+
+	TArray<FString> AbsoluteFiles = SourceControlHelpers::AbsoluteFilenames(InFiles);
+
+	if(InStateCacheUsage == EStateCacheUsage::ForceUpdate)
+	{
+		UE_LOG(LogSourceControl, Log, TEXT("GetState: ForceUpdate"));
+		Execute(ISourceControlOperation::Create<FUpdateStatus>(), AbsoluteFiles);
+	}
+
+	for(const auto& AbsoluteFile : AbsoluteFiles)
+	{
+		OutState.Add(GetStateInternal(*AbsoluteFile));
+	}
+
+	return ECommandResult::Succeeded;
+}
+
+TArray<FSourceControlStateRef> FPlasticSourceControlProvider::GetCachedStateByPredicate(TFunctionRef<bool(const FSourceControlStateRef&)> Predicate) const
+{
+	TArray<FSourceControlStateRef> Result;
+	for(const auto& CacheItem : StateCache)
+	{
+		FSourceControlStateRef State = CacheItem.Value;
+		if(Predicate(State))
+		{
+			Result.Add(State);
+		}
+	}
+	return Result;
+}
+
+bool FPlasticSourceControlProvider::RemoveFileFromCache(const FString& Filename)
+{
+	return StateCache.Remove(Filename) > 0;
+}
+
+FDelegateHandle FPlasticSourceControlProvider::RegisterSourceControlStateChanged_Handle( const FSourceControlStateChanged::FDelegate& SourceControlStateChanged )
+{
+	return OnSourceControlStateChanged.Add( SourceControlStateChanged );
+}
+
+void FPlasticSourceControlProvider::UnregisterSourceControlStateChanged_Handle( FDelegateHandle Handle )
+{
+	OnSourceControlStateChanged.Remove( Handle );
+}
+
+ECommandResult::Type FPlasticSourceControlProvider::Execute( const TSharedRef<ISourceControlOperation, ESPMode::ThreadSafe>& InOperation, const TArray<FString>& InFiles, EConcurrency::Type InConcurrency, const FSourceControlOperationComplete& InOperationCompleteDelegate )
+{
+	if(!bWorkspaceFound && !(InOperation->GetName() == "Connect") && !(InOperation->GetName() == "MakeWorkspace"))
+	{
+		UE_LOG(LogSourceControl, Warning, TEXT("'%s': only Connect operation allowed whithout a workspace"), *InOperation->GetName().ToString());
+		return ECommandResult::Failed;
+	}
+
+	TArray<FString> AbsoluteFiles = SourceControlHelpers::AbsoluteFilenames(InFiles);
+
+	// Query to see if we allow this operation
+	TSharedPtr<IPlasticSourceControlWorker, ESPMode::ThreadSafe> Worker = CreateWorker(InOperation->GetName());
+	if(!Worker.IsValid())
+	{
+		// this operation is unsupported by this source control provider
+		FFormatNamedArguments Arguments;
+		Arguments.Add( TEXT("OperationName"), FText::FromName(InOperation->GetName()) );
+		Arguments.Add( TEXT("ProviderName"), FText::FromName(GetName()) );
+		FMessageLog("SourceControl").Error(FText::Format(LOCTEXT("UnsupportedOperation", "Operation '{OperationName}' not supported by source control provider '{ProviderName}'"), Arguments));
+		return ECommandResult::Failed;
+	}
+
+	FPlasticSourceControlCommand* Command = new FPlasticSourceControlCommand(InOperation, Worker.ToSharedRef());
+	Command->Files = AbsoluteFiles;
+	Command->OperationCompleteDelegate = InOperationCompleteDelegate;
+
+	// fire off operation
+	if(InConcurrency == EConcurrency::Synchronous)
+	{
+		Command->bAutoDelete = false;
+
+		UE_LOG(LogSourceControl, Log, TEXT("ExecuteSynchronousCommand: %s"), *InOperation->GetName().ToString());
+		return ExecuteSynchronousCommand(*Command, InOperation->GetInProgressString());
+	}
+	else
+	{
+		Command->bAutoDelete = true;
+
+		UE_LOG(LogSourceControl, Log, TEXT("IssueAsynchronousCommand: %s"), *InOperation->GetName().ToString());
+		return IssueCommand(*Command);
+	}
+}
+
+bool FPlasticSourceControlProvider::CanCancelOperation( const TSharedRef<ISourceControlOperation, ESPMode::ThreadSafe>& InOperation ) const
+{
+	return false;
+}
+
+void FPlasticSourceControlProvider::CancelOperation( const TSharedRef<ISourceControlOperation, ESPMode::ThreadSafe>& InOperation )
+{
+}
+
+bool FPlasticSourceControlProvider::UsesLocalReadOnlyState() const
+{
+	return false; // TODO: use configuration
+}
+
+bool FPlasticSourceControlProvider::UsesChangelists() const
+{
+	return false; // We don't want to show ChangeList column anymore (Plastic SCM term would be ChangeSet)
+}
+
+bool FPlasticSourceControlProvider::UsesCheckout() const
+{
+	return true;
+}
+
+TSharedPtr<IPlasticSourceControlWorker, ESPMode::ThreadSafe> FPlasticSourceControlProvider::CreateWorker(const FName& InOperationName) const
+{
+	const FGetPlasticSourceControlWorker* Operation = WorkersMap.Find(InOperationName);
+	if(Operation != nullptr)
+	{
+		return Operation->Execute();
+	}
+		
+	return nullptr;
+}
+
+void FPlasticSourceControlProvider::RegisterWorker( const FName& InName, const FGetPlasticSourceControlWorker& InDelegate )
+{
+	WorkersMap.Add( InName, InDelegate );
+}
+
+void FPlasticSourceControlProvider::OutputCommandMessages(const FPlasticSourceControlCommand& InCommand) const
+{
+	FMessageLog SourceControlLog("SourceControl");
+
+	for(int32 ErrorIndex = 0; ErrorIndex < InCommand.ErrorMessages.Num(); ++ErrorIndex)
+	{
+		SourceControlLog.Error(FText::FromString(InCommand.ErrorMessages[ErrorIndex]));
+	}
+
+	for(int32 InfoIndex = 0; InfoIndex < InCommand.InfoMessages.Num(); ++InfoIndex)
+	{
+		SourceControlLog.Info(FText::FromString(InCommand.InfoMessages[InfoIndex]));
+	}
+}
+
+void FPlasticSourceControlProvider::UpdateWorkspaceStatus(const class FPlasticSourceControlCommand& InCommand)
+{
+	if (InCommand.Operation->GetName() == "Connect")
+	{
+		// Is connection successful?
+		bServerAvailable = InCommand.bCommandSuccessful;
+		bWorkspaceFound = (InCommand.WorkspaceName.Len() > 0);
+
+		WorkspaceName = InCommand.WorkspaceName;
+		RepositoryName = InCommand.RepositoryName;
+		ServerUrl = InCommand.ServerUrl;
+
+		if (bWorkspaceFound)
+		{
+			// Extend the "Source Control" menu in the Editor Toolbar on each successful connection
+			PlasticSourceControlMenu.Unregister(); // cleanup for any previous connection
+			PlasticSourceControlMenu.Register();
+		}
+	}
+	else if (InCommand.bConnectionDropped)
+	{
+		// checkconnection failed on UpdateStatus
+		bServerAvailable = false;
+	}
+
+	// And for all operations running UpdateStatus, get Changeset and Branch informations:
+	if (!InCommand.BranchName.IsEmpty())
+	{
+		ChangesetNumber = InCommand.ChangesetNumber;
+		BranchName = InCommand.BranchName;
+	}
+}
+
+void FPlasticSourceControlProvider::Tick()
+{	
+	bool bStatesUpdated = false;
+	for(int32 CommandIndex = 0; CommandIndex < CommandQueue.Num(); ++CommandIndex)
+	{
+		FPlasticSourceControlCommand& Command = *CommandQueue[CommandIndex];
+		if(Command.bExecuteProcessed)
+		{
+			// Remove command from the queue
+			CommandQueue.RemoveAt(CommandIndex);
+
+			// Update workspace status and connection state on Connect and UpdateStatus operations
+			UpdateWorkspaceStatus(Command);
+
+			// let command update the states of any files
+			bStatesUpdated |= Command.Worker->UpdateStates();
+
+			// dump any messages to output log
+			OutputCommandMessages(Command);
+
+			// run the completion delegate callback if we have one bound
+			ECommandResult::Type Result = Command.bCommandSuccessful ? ECommandResult::Succeeded : ECommandResult::Failed;
+			Command.OperationCompleteDelegate.ExecuteIfBound(Command.Operation, Result);
+
+			// commands that are left in the array during a tick need to be deleted
+			if(Command.bAutoDelete)
+			{
+				// Only delete commands that are not running 'synchronously'
+				delete &Command;
+			}
+
+			// only do one command per tick loop, as we dont want concurrent modification 
+			// of the command queue (which can happen in the completion delegate)
+			break;
+		}
+	}
+
+	if(bStatesUpdated)
+	{
+		OnSourceControlStateChanged.Broadcast();
+	}
+}
+
+TArray< TSharedRef<ISourceControlLabel> > FPlasticSourceControlProvider::GetLabels( const FString& InMatchingSpec ) const
+{
+	TArray< TSharedRef<ISourceControlLabel> > Tags;
+
+	// NOTE list labels. Called by CrashDebugHelper() (to remote debug Engine crash)
+	//					 and by SourceControlHelpers::AnnotateFile() (to add source file to report)
+	// Reserved for internal use by Epic Games with Perforce only
+	return Tags;
+}
+
+#if SOURCE_CONTROL_WITH_SLATE
+TSharedRef<class SWidget> FPlasticSourceControlProvider::MakeSettingsWidget() const
+{
+	return SNew(SPlasticSourceControlSettings);
+}
+#endif
+
+ECommandResult::Type FPlasticSourceControlProvider::ExecuteSynchronousCommand(FPlasticSourceControlCommand& InCommand, const FText& Task)
+{
+	ECommandResult::Type Result = ECommandResult::Failed;
+
+	// Display the progress dialog if a string was provided
+	{
+		FScopedSourceControlProgress Progress(Task);
+
+		// Issue the command asynchronously...
+		IssueCommand( InCommand );
+
+		// ... then wait for its completion (thus making it synchronous)
+		while(!InCommand.bExecuteProcessed)
+		{
+			// Tick the command queue and update progress.
+			Tick();
+
+			Progress.Tick();
+
+			// Sleep for a bit so we don't busy-wait so much.
+			FPlatformProcess::Sleep(0.01f);
+		}
+
+		// always do one more Tick() to make sure the command queue is cleaned up.
+		Tick();
+
+		if(InCommand.bCommandSuccessful)
+		{
+			Result = ECommandResult::Succeeded;
+		}
+		else
+		{
+			// TODO If the command failed, inform the user that they need to try again (see Perforce)
+			//FMessageDialog::Open( EAppMsgType::Ok, LOCTEXT("Plastic_ServerUnresponsive", "Plastic server is unresponsive. Please check your connection and try again.") );
+
+			UE_LOG(LogSourceControl, Error, TEXT("Command '%s' Failed!"), *InCommand.Operation->GetName().ToString());
+		}
+	}
+
+	// Delete the command now (asynchronous commands are deleted in the Tick() method)
+	check(!InCommand.bAutoDelete);
+
+	// ensure commands that are not auto deleted do not end up in the command queue
+	if ( CommandQueue.Contains( &InCommand ) ) 
+	{
+		CommandQueue.Remove( &InCommand );
+	}
+	delete &InCommand;
+
+	return Result;
+}
+
+ECommandResult::Type FPlasticSourceControlProvider::IssueCommand(FPlasticSourceControlCommand& InCommand)
+{
+	if(GThreadPool != nullptr)
+	{
+		// Queue this to our worker thread(s) for resolving
+		GThreadPool->AddQueuedWork(&InCommand);
+		CommandQueue.Add(&InCommand);
+		return ECommandResult::Succeeded;
+	}
+	else
+	{
+		return ECommandResult::Failed;
+	}
+}
+
+#undef LOCTEXT_NAMESPACE
diff --git a/Engine/Plugins/Developer/PlasticSourceControl/Source/PlasticSourceControl/Private/PlasticSourceControlProvider.h b/Engine/Plugins/Developer/PlasticSourceControl/Source/PlasticSourceControl/Private/PlasticSourceControlProvider.h
new file mode 100644
--- /dev/null
+++ b/Engine/Plugins/Developer/PlasticSourceControl/Source/PlasticSourceControl/Private/PlasticSourceControlProvider.h
@@ -0,0 +1,185 @@
+// Copyright 1998-2019 Epic Games, Inc. All Rights Reserved.
+
+#pragma once
+
+#include "CoreMinimal.h"
+#include "ISourceControlProvider.h"
+#include "IPlasticSourceControlWorker.h"
+#include "PlasticSourceControlState.h"
+#include "PlasticSourceControlMenu.h"
+
+DECLARE_DELEGATE_RetVal(FPlasticSourceControlWorkerRef, FGetPlasticSourceControlWorker)
+
+class FPlasticSourceControlProvider : public ISourceControlProvider
+{
+public:
+	/** Constructor */
+	FPlasticSourceControlProvider()
+		: bPlasticAvailable(false)
+		, bWorkspaceFound(false)
+		, bServerAvailable(false)
+		, ChangesetNumber(0)
+	{
+	}
+
+	/* ISourceControlProvider implementation */
+	virtual void Init(bool bForceConnection = true) override;
+	virtual void Close() override;
+	virtual FText GetStatusText() const override;
+	virtual bool IsEnabled() const override;
+	virtual bool IsAvailable() const override;
+	virtual const FName& GetName(void) const override;
+	virtual bool QueryStateBranchConfig(const FString& ConfigSrc, const FString& ConfigDest) /* override UE4.20 */ { return false; }
+	virtual void RegisterStateBranches(const TArray<FString>& BranchNames, const FString& ContentRoot) /* override UE4.20 */ {}
+	virtual int32 GetStateBranchIndex(const FString& InBranchName) const /* override UE4.20 */ { return INDEX_NONE; }
+	virtual ECommandResult::Type GetState( const TArray<FString>& InFiles, TArray< TSharedRef<ISourceControlState, ESPMode::ThreadSafe> >& OutState, EStateCacheUsage::Type InStateCacheUsage ) override;
+	virtual TArray<FSourceControlStateRef> GetCachedStateByPredicate(TFunctionRef<bool(const FSourceControlStateRef&)> Predicate) const override;
+	virtual FDelegateHandle RegisterSourceControlStateChanged_Handle( const FSourceControlStateChanged::FDelegate& SourceControlStateChanged ) override;
+	virtual void UnregisterSourceControlStateChanged_Handle( FDelegateHandle Handle ) override;
+	virtual ECommandResult::Type Execute( const TSharedRef<ISourceControlOperation, ESPMode::ThreadSafe>& InOperation, const TArray<FString>& InFiles, EConcurrency::Type InConcurrency = EConcurrency::Synchronous, const FSourceControlOperationComplete& InOperationCompleteDelegate = FSourceControlOperationComplete() ) override;
+	virtual bool CanCancelOperation( const TSharedRef<ISourceControlOperation, ESPMode::ThreadSafe>& InOperation ) const override;
+	virtual void CancelOperation( const TSharedRef<ISourceControlOperation, ESPMode::ThreadSafe>& InOperation ) override;
+	virtual bool UsesLocalReadOnlyState() const override;
+	virtual bool UsesChangelists() const override;
+	virtual bool UsesCheckout() const override;
+	virtual void Tick() override;
+	virtual TArray< TSharedRef<class ISourceControlLabel> > GetLabels( const FString& InMatchingSpec ) const override;
+#if SOURCE_CONTROL_WITH_SLATE
+	virtual TSharedRef<class SWidget> MakeSettingsWidget() const override;
+#endif
+
+	/**
+	 * Run a Plastic "version" command to check the availability of the binary and of the workspace.
+	 */
+	void CheckPlasticAvailability();
+
+	/** Is Plastic command line working. */
+	inline bool IsPlasticAvailable() const
+	{
+		return bPlasticAvailable;
+	}
+
+	/** Is Plastic workspace found. */
+	inline bool IsWorkspaceFound() const
+	{
+		return bWorkspaceFound;
+	}
+
+	/** Get the path to the root of the Plastic workspace: can be the GameDir itself, or any parent directory */
+	inline const FString& GetPathToWorkspaceRoot() const
+	{
+		return PathToWorkspaceRoot;
+	}
+
+	/** Get the Plastic current user */
+	inline const FString& GetUserName() const
+	{
+		return UserName;
+	}
+
+	/** Get the Plastic current workspace */
+	inline const FString& GetWorkspaceName() const
+	{
+		return WorkspaceName;
+	}
+
+	/** Get the Plastic current repository */
+	inline const FString& GetRepositoryName() const
+	{
+		return RepositoryName;
+	}
+
+	/** Get the Plastic current server URL */
+	inline const FString& GetServerUrl() const
+	{
+		return ServerUrl;
+	}
+
+	/** Get the Name of the current branch */
+	inline const FString& GetBranchName() const
+	{
+		return BranchName;
+	}
+
+	/** Get the current Changeset Number */
+	inline int32 GetChangesetNumber() const
+	{
+		return ChangesetNumber;
+	}
+
+	/** Helper function used to update state cache */
+	TSharedRef<FPlasticSourceControlState, ESPMode::ThreadSafe> GetStateInternal(const FString& Filename) const;
+
+	/**
+	 * Register a worker with the provider.
+	 * This is used internally so the provider can maintain a map of all available operations.
+	 */
+	void RegisterWorker( const FName& InName, const FGetPlasticSourceControlWorker& InDelegate );
+
+	/** Remove a named file from the state cache */
+	bool RemoveFileFromCache(const FString& Filename);
+
+private:
+
+	/** Is Plastic binary found and working. */
+	bool bPlasticAvailable;
+
+	/** Is Plastic workspace found. */
+	bool bWorkspaceFound;
+
+	/** Indicates if source control integration is available or not. */
+	bool bServerAvailable;
+
+	/** Helper function for Execute() */
+	TSharedPtr<class IPlasticSourceControlWorker, ESPMode::ThreadSafe> CreateWorker(const FName& InOperationName) const;
+
+	/** Helper function for running command synchronously. */
+	ECommandResult::Type ExecuteSynchronousCommand(class FPlasticSourceControlCommand& InCommand, const FText& Task);
+	/** Issue a command asynchronously if possible. */
+	ECommandResult::Type IssueCommand(class FPlasticSourceControlCommand& InCommand);
+
+	/** Output any messages this command holds */
+	void OutputCommandMessages(const class FPlasticSourceControlCommand& InCommand) const;
+
+	/** Update workspace status on Connect and UpdateStatus operations */
+	void UpdateWorkspaceStatus(const class FPlasticSourceControlCommand& InCommand);
+
+    /** Version of the Plastic SCM executable used */
+    FString PlasticScmVersion;
+
+	/** Path to the root of the Plastic workspace: can be the GameDir itself, or any parent directory (found by the "Connect" operation) */
+	FString PathToWorkspaceRoot;
+
+	/** Plastic current user */
+	FString UserName;
+
+	/** Plastic current workspace */
+	FString WorkspaceName;
+
+	/** Plastic current repository */
+	FString RepositoryName;
+
+	/** Plastic current server URL */
+	FString ServerUrl;
+
+	/** Name of the current branch */
+	FString BranchName;
+
+	/** Current Changeset Number */
+	int32 ChangesetNumber;
+
+	/** State cache */
+	mutable TMap<FString, TSharedRef<class FPlasticSourceControlState, ESPMode::ThreadSafe> > StateCache;
+
+	/** The currently registered source control operations */
+	TMap<FName, FGetPlasticSourceControlWorker> WorkersMap;
+
+	/** Queue for commands given by the main thread */
+	TArray < FPlasticSourceControlCommand* > CommandQueue;
+
+	/** For notifying when the source control states in the cache have changed */
+	FSourceControlStateChanged OnSourceControlStateChanged;
+
+	/** Source Control Menu Extension */
+	FPlasticSourceControlMenu PlasticSourceControlMenu;
+};
diff --git a/Engine/Plugins/Developer/PlasticSourceControl/Source/PlasticSourceControl/Private/PlasticSourceControlRevision.cpp b/Engine/Plugins/Developer/PlasticSourceControl/Source/PlasticSourceControl/Private/PlasticSourceControlRevision.cpp
new file mode 100644
--- /dev/null
+++ b/Engine/Plugins/Developer/PlasticSourceControl/Source/PlasticSourceControl/Private/PlasticSourceControlRevision.cpp
@@ -0,0 +1,122 @@
+// Copyright 1998-2019 Epic Games, Inc. All Rights Reserved.
+
+#include "PlasticSourceControlRevision.h"
+#include "PlasticSourceControlModule.h"
+#include "PlasticSourceControlProvider.h"
+#include "PlasticSourceControlState.h"
+#include "PlasticSourceControlUtils.h"
+#include "SPlasticSourceControlSettings.h"
+
+#include "HAL/FileManager.h"
+#include "Misc/Paths.h"
+
+#define LOCTEXT_NAMESPACE "PlasticSourceControl"
+
+bool FPlasticSourceControlRevision::Get( FString& InOutFilename ) const
+{
+	const FPlasticSourceControlModule& PlasticSourceControl = FModuleManager::LoadModuleChecked<FPlasticSourceControlModule>("PlasticSourceControl");
+	const FString& PathToPlasticBinary = PlasticSourceControl.AccessSettings().GetBinaryPath();
+	if (!State)
+	{
+		UE_LOG(LogSourceControl, Error, TEXT("Revision(%s %d): unknown state!"), *Filename, RevisionId);
+	}
+
+	// if a filename for the temp file wasn't supplied generate a unique-ish one
+	if(InOutFilename.Len() == 0)
+	{
+		// create the diff dir if we don't already have it (Plastic wont)
+		IFileManager::Get().MakeDirectory(*FPaths::DiffDir(), true);
+		// create a unique temp file name based on the unique revision Id
+		const FString TempFileName = FString::Printf(TEXT("%stemp-%d-%s"), *FPaths::DiffDir(), RevisionId, *FPaths::GetCleanFilename(Filename));
+		InOutFilename = FPaths::ConvertRelativePathToFull(TempFileName);
+	}
+
+	bool bCommandSuccessful;
+	if(FPaths::FileExists(InOutFilename))
+	{
+		bCommandSuccessful = true; // if the temp file already exists, reuse it directly
+	}
+	else
+	{
+		// Format the revision specification of the file, like revid:1230@repo@server:8087
+		const FString RevisionSpecification = FString::Printf(TEXT("revid:%d@%s"), RevisionId, *State->RepSpec);
+		bCommandSuccessful = PlasticSourceControlUtils::RunDumpToFile(PathToPlasticBinary, RevisionSpecification, InOutFilename);
+	}
+	return bCommandSuccessful;
+}
+
+bool FPlasticSourceControlRevision::GetAnnotated( TArray<FAnnotationLine>& OutLines ) const
+{
+	// NOTE GetAnnotated: called only by SourceControlHelpers::AnnotateFile(),
+	//      called only by ICrashDebugHelper::AddAnnotatedSourceToReport() using a changelist/check identifier
+	//      called only by FCrashDebugHelperWindows::CreateMinidumpDiagnosticReport() (and Mac) to Extract annotated lines from a source file stored in Perforce, and add to the crash report.
+	//      called by - MinidumpDiagnosticsApp RunMinidumpDiagnostics() for Perfore ONLY "MinidumpDiagnostics.exe <Crash.dmp> [-Annotate] [-SyncSymbols] [-SyncMicrosoftSymbols]"
+	//                - FWindowsErrorReport::DiagnoseReport() (and Mac)
+	// Reserved for internal use by Epic Games with Perforce only
+	return false;
+}
+
+bool FPlasticSourceControlRevision::GetAnnotated( FString& InOutFilename ) const
+{
+	// NOTE: Unused, only the above method is called by the Editor
+	return false;
+}
+
+const FString& FPlasticSourceControlRevision::GetFilename() const
+{
+	return Filename;
+}
+
+int32 FPlasticSourceControlRevision::GetRevisionNumber() const
+{
+	return ChangesetNumber; // Using the Changelist as the Revision number to display in the Asset Diff Menu
+}
+
+const FString& FPlasticSourceControlRevision::GetRevision() const
+{
+	return Revision;
+}
+
+const FString& FPlasticSourceControlRevision::GetDescription() const
+{
+	return Description;
+}
+
+const FString& FPlasticSourceControlRevision::GetUserName() const
+{
+	return UserName;
+}
+
+const FString& FPlasticSourceControlRevision::GetClientSpec() const
+{
+	static FString EmptyString(TEXT("")); // NOTE Workspace/Clientspec of the submitter (Perforce only)
+	return EmptyString;
+}
+
+const FString& FPlasticSourceControlRevision::GetAction() const
+{
+	return Action;
+}
+
+TSharedPtr<class ISourceControlRevision, ESPMode::ThreadSafe> FPlasticSourceControlRevision::GetBranchSource() const
+{
+	// if this revision was copied/moved from some other revision
+	return BranchSource;
+}
+
+const FDateTime& FPlasticSourceControlRevision::GetDate() const
+{
+	return Date;
+}
+
+int32 FPlasticSourceControlRevision::GetCheckInIdentifier() const
+{
+	return ChangesetNumber;
+}
+
+int32 FPlasticSourceControlRevision::GetFileSize() const
+{
+	return FileSize;
+}
+
+#undef LOCTEXT_NAMESPACE
diff --git a/Engine/Plugins/Developer/PlasticSourceControl/Source/PlasticSourceControl/Private/PlasticSourceControlRevision.h b/Engine/Plugins/Developer/PlasticSourceControl/Source/PlasticSourceControl/Private/PlasticSourceControlRevision.h
new file mode 100644
--- /dev/null
+++ b/Engine/Plugins/Developer/PlasticSourceControl/Source/PlasticSourceControl/Private/PlasticSourceControlRevision.h
@@ -0,0 +1,75 @@
+// Copyright 1998-2016 Epic Games, Inc. All Rights Reserved.
+
+#pragma once
+
+#include "CoreMinimal.h"
+#include "ISourceControlRevision.h"
+
+class FPlasticSourceControlState;
+
+/** Revision of a file, linked to a specific commit */
+class FPlasticSourceControlRevision : public ISourceControlRevision, public TSharedFromThis<FPlasticSourceControlRevision, ESPMode::ThreadSafe>
+{
+public:
+	FPlasticSourceControlRevision()
+		: ChangesetNumber(0)
+		, RevisionId(0)
+		, Date(0)
+		, FileSize(0)
+	{
+	}
+
+	/** ISourceControlRevision interface */
+	virtual bool Get( FString& InOutFilename ) const override;
+	virtual bool GetAnnotated( TArray<FAnnotationLine>& OutLines ) const override;
+	virtual bool GetAnnotated( FString& InOutFilename ) const override;
+	virtual const FString& GetFilename() const override;
+	virtual int32 GetRevisionNumber() const override;
+	virtual const FString& GetRevision() const override;
+	virtual const FString& GetDescription() const override;
+	virtual const FString& GetUserName() const override;
+	virtual const FString& GetClientSpec() const override;
+	virtual const FString& GetAction() const override;
+	virtual TSharedPtr<class ISourceControlRevision, ESPMode::ThreadSafe> GetBranchSource() const override;
+	virtual const FDateTime& GetDate() const override;
+	virtual int32 GetCheckInIdentifier() const override;
+	virtual int32 GetFileSize() const override;
+
+public:
+
+	/** Point back to State this Revision is from */
+	FPlasticSourceControlState* State = nullptr;
+
+	/** The filename this revision refers to */
+	FString Filename;
+
+	/** The changeset number of this revision */
+	int32 ChangesetNumber;
+
+	/** The internal revision ID of this file */
+	int32 RevisionId;
+
+	/** The revision to display to the user: use the changeset number */
+	FString Revision;
+
+	/** The description of this revision */
+	FString Description;
+
+	/** The user that made the change */
+	FString UserName;
+
+	/** The action (add, edit, branch etc.) performed at this revision */
+	FString Action;
+
+	/** Source of move ("branch" in Perforce term) if any */
+	TSharedPtr<FPlasticSourceControlRevision, ESPMode::ThreadSafe> BranchSource;
+
+	/** The date this revision was made */
+	FDateTime Date;
+
+	/** The size of the file at this revision */
+	int32 FileSize;
+};
+
+/** History composed of the last 100 revisions of the file */
+typedef TArray<TSharedRef<FPlasticSourceControlRevision, ESPMode::ThreadSafe>>	TPlasticSourceControlHistory;
diff --git a/Engine/Plugins/Developer/PlasticSourceControl/Source/PlasticSourceControl/Private/PlasticSourceControlSettings.cpp b/Engine/Plugins/Developer/PlasticSourceControl/Source/PlasticSourceControl/Private/PlasticSourceControlSettings.cpp
new file mode 100644
--- /dev/null
+++ b/Engine/Plugins/Developer/PlasticSourceControl/Source/PlasticSourceControl/Private/PlasticSourceControlSettings.cpp
@@ -0,0 +1,65 @@
+// Copyright 1998-2019 Epic Games, Inc. All Rights Reserved.
+
+#include "PlasticSourceControlSettings.h"
+#include "PlasticSourceControlModule.h"
+#include "PlasticSourceControlProvider.h"
+#include "PlasticSourceControlUtils.h"
+#include "SourceControlHelpers.h"
+
+#include "Misc/ScopeLock.h"
+#include "Misc/ConfigCacheIni.h"
+#include "HAL/PlatformTime.h"
+
+namespace PlasticSettingsConstants
+{
+
+/** The section of the ini file we load our settings from */
+static const FString SettingsSection = TEXT("PlasticSourceControl.PlasticSourceControlSettings");
+
+}
+
+FString FPlasticSourceControlSettings::GetBinaryPath() const
+{
+	FScopeLock ScopeLock(&CriticalSection);
+	return BinaryPath;
+}
+
+bool FPlasticSourceControlSettings::SetBinaryPath(const FString& InString)
+{
+	FScopeLock ScopeLock(&CriticalSection);
+	const bool bChanged = (BinaryPath != InString);
+	if(bChanged)
+	{
+		BinaryPath = InString;
+	}
+	return bChanged;
+}
+
+bool FPlasticSourceControlSettings::UpdateStatusAtStartup() const
+{
+	FScopeLock ScopeLock(&CriticalSection);
+	return bUpdateStatusAtStartup;
+}
+
+void FPlasticSourceControlSettings::SetUpdateStatusAtStartup(const bool bInUpdateStatusAtStartup)
+{
+	FScopeLock ScopeLock(&CriticalSection);
+	bUpdateStatusAtStartup = bInUpdateStatusAtStartup;
+}
+
+// This is called at startup nearly before anything else in our module: BinaryPath will then be used by the provider
+void FPlasticSourceControlSettings::LoadSettings()
+{
+	FScopeLock ScopeLock(&CriticalSection);
+	const FString& IniFile = SourceControlHelpers::GetSettingsIni();
+	GConfig->GetString(*PlasticSettingsConstants::SettingsSection, TEXT("BinaryPath"), BinaryPath, IniFile);
+	GConfig->GetBool(*PlasticSettingsConstants::SettingsSection, TEXT("UpdateStatusAtStartup"), bUpdateStatusAtStartup, IniFile);
+}
+
+void FPlasticSourceControlSettings::SaveSettings() const
+{
+	FScopeLock ScopeLock(&CriticalSection);
+	const FString& IniFile = SourceControlHelpers::GetSettingsIni();
+	GConfig->SetString(*PlasticSettingsConstants::SettingsSection, TEXT("BinaryPath"), *BinaryPath, IniFile);
+	GConfig->SetBool(*PlasticSettingsConstants::SettingsSection, TEXT("UpdateStatusAtStartup"), bUpdateStatusAtStartup, IniFile);
+}
\ No newline at end of file
diff --git a/Engine/Plugins/Developer/PlasticSourceControl/Source/PlasticSourceControl/Private/PlasticSourceControlSettings.h b/Engine/Plugins/Developer/PlasticSourceControl/Source/PlasticSourceControl/Private/PlasticSourceControlSettings.h
new file mode 100644
--- /dev/null
+++ b/Engine/Plugins/Developer/PlasticSourceControl/Source/PlasticSourceControl/Private/PlasticSourceControlSettings.h
@@ -0,0 +1,36 @@
+// Copyright 1998-2019 Epic Games, Inc. All Rights Reserved.
+
+#pragma once
+
+#include "CoreMinimal.h"
+
+class FPlasticSourceControlSettings
+{
+public:
+	/** Get/Set the Plastic Binary Path */
+	FString GetBinaryPath() const;
+	bool SetBinaryPath(const FString& InString);
+
+	/** Enable an asynchronous "Update Status" at Editor Startup (default is no, can take a long time). */
+	bool UpdateStatusAtStartup() const;
+	void SetUpdateStatusAtStartup(const bool bInUpdateStatusAtStartup);
+
+	/** Load settings from ini file */
+	void LoadSettings();
+
+	/** Save settings to ini file */
+	void SaveSettings() const;
+
+private:
+	/** A critical section for settings access */
+	mutable FCriticalSection CriticalSection;
+
+	/** Plastic binary path */
+	FString BinaryPath;
+
+	/** Enable an asynchronous "Update Status" at Editor Startup (default is no).
+	 * This does not work well with very big projects where this operation could take dozens of seconds
+	 * preventing the project to have any source control support during this time.
+	*/
+	bool bUpdateStatusAtStartup;
+};
\ No newline at end of file
diff --git a/Engine/Plugins/Developer/PlasticSourceControl/Source/PlasticSourceControl/Private/PlasticSourceControlState.cpp b/Engine/Plugins/Developer/PlasticSourceControl/Source/PlasticSourceControl/Private/PlasticSourceControlState.cpp
new file mode 100644
--- /dev/null
+++ b/Engine/Plugins/Developer/PlasticSourceControl/Source/PlasticSourceControl/Private/PlasticSourceControlState.cpp
@@ -0,0 +1,456 @@
+// Copyright 1998-2019 Epic Games, Inc. All Rights Reserved.
+
+#include "PlasticSourceControlState.h"
+
+#define LOCTEXT_NAMESPACE "PlasticSourceControl.State"
+
+namespace EWorkspaceState
+{
+const TCHAR* ToString(EWorkspaceState::Type InWorkspaceState)
+{
+	const TCHAR* WorkspaceStateStr = nullptr;
+	switch (InWorkspaceState)
+	{
+	case EWorkspaceState::Unknown: WorkspaceStateStr = TEXT("Unknown"); break;
+	case EWorkspaceState::Ignored: WorkspaceStateStr = TEXT("Ignored"); break;
+	case EWorkspaceState::Controlled: WorkspaceStateStr = TEXT("Controlled"); break;
+	case EWorkspaceState::CheckedOut: WorkspaceStateStr = TEXT("CheckedOut"); break;
+	case EWorkspaceState::Added: WorkspaceStateStr = TEXT("Added"); break;
+	case EWorkspaceState::Moved: WorkspaceStateStr = TEXT("Moved"); break;
+	case EWorkspaceState::Copied: WorkspaceStateStr = TEXT("Copied"); break;
+	case EWorkspaceState::Replaced: WorkspaceStateStr = TEXT("Replaced"); break;
+	case EWorkspaceState::Deleted: WorkspaceStateStr = TEXT("Deleted"); break;
+	case EWorkspaceState::LocallyDeleted: WorkspaceStateStr = TEXT("LocallyDeleted"); break;
+	case EWorkspaceState::Changed: WorkspaceStateStr = TEXT("Changed"); break;
+	case EWorkspaceState::Conflicted: WorkspaceStateStr = TEXT("Conflicted"); break;
+	case EWorkspaceState::LockedByOther: WorkspaceStateStr = TEXT("LockedByOther"); break;
+	case EWorkspaceState::Private: WorkspaceStateStr = TEXT("Private"); break;
+	default: WorkspaceStateStr = TEXT("???"); break;
+	}
+	return WorkspaceStateStr;
+}
+}
+
+
+int32 FPlasticSourceControlState::GetHistorySize() const
+{
+	return History.Num();
+}
+
+TSharedPtr<class ISourceControlRevision, ESPMode::ThreadSafe> FPlasticSourceControlState::GetHistoryItem( int32 HistoryIndex ) const
+{
+	check(History.IsValidIndex(HistoryIndex));
+	return History[HistoryIndex];
+}
+
+TSharedPtr<class ISourceControlRevision, ESPMode::ThreadSafe> FPlasticSourceControlState::FindHistoryRevision( int32 RevisionNumber ) const
+{
+	for(const auto& Revision : History)
+	{
+		if(Revision->GetRevisionNumber() == RevisionNumber)
+		{
+			return Revision;
+		}
+	}
+
+	return nullptr;
+}
+
+TSharedPtr<class ISourceControlRevision, ESPMode::ThreadSafe> FPlasticSourceControlState::FindHistoryRevision(const FString& InRevision) const
+{
+	for(const auto& Revision : History)
+	{
+		if(Revision->GetRevision() == InRevision)
+		{
+			return Revision;
+		}
+	}
+
+	return nullptr;
+}
+
+TSharedPtr<class ISourceControlRevision, ESPMode::ThreadSafe> FPlasticSourceControlState::GetBaseRevForMerge() const
+{
+	for(const auto& Revision : History)
+	{
+		// look for the the SHA1 id of the file, not the commit id (revision)
+		if(Revision->ChangesetNumber == PendingMergeBaseChangeset)
+		{
+			return Revision;
+		}
+	}
+
+	return nullptr;
+}
+
+FName FPlasticSourceControlState::GetIconName() const
+{
+	if (!IsCurrent())
+	{
+		return FName("Perforce.NotAtHeadRevision");
+	}
+
+	switch(WorkspaceState)
+	{
+	case EWorkspaceState::CheckedOut:
+	case EWorkspaceState::Replaced: // Merged (waiting for checkin) TODO: would need a dedicated icon
+		return FName("Perforce.CheckedOut");
+	case EWorkspaceState::Added:
+	case EWorkspaceState::Copied:
+		return FName("Perforce.OpenForAdd");
+	case EWorkspaceState::Moved:
+		return FName("Perforce.Branched");
+	case EWorkspaceState::Deleted: // Deleted & Missing files does not show in Content Browser
+		return FName("Perforce.MarkedForDelete");
+	case EWorkspaceState::LocallyDeleted: // Deleted & Missing files does not show in Content Browser
+	case EWorkspaceState::Conflicted:
+		return FName("Perforce.NotAtHeadRevision");
+	case EWorkspaceState::LockedByOther:
+		return FName("Perforce.CheckedOutByOtherUser");
+	case EWorkspaceState::Private: // Not controlled
+	case EWorkspaceState::Changed: // Changed but unchecked-out file is in a certain way not controlled - TODO: would need a dedicated icon
+		return FName("Perforce.NotInDepot");
+	case EWorkspaceState::Unknown:
+	case EWorkspaceState::Ignored:
+	case EWorkspaceState::Controlled: // (Unchanged) same as "Pristine" for Perforce (not checked out) ie no icon
+	default:
+		return NAME_None;
+	}
+}
+
+FName FPlasticSourceControlState::GetSmallIconName() const
+{
+	if (!IsCurrent())
+	{
+		return FName("Perforce.NotAtHeadRevision_Small");
+	}
+
+	switch(WorkspaceState)
+	{
+	case EWorkspaceState::CheckedOut:
+	case EWorkspaceState::Replaced: // Merged (waiting for checkin)
+		return FName("Perforce.CheckedOut_Small");
+	case EWorkspaceState::Added:
+	case EWorkspaceState::Copied:
+		return FName("Perforce.OpenForAdd_Small");
+	case EWorkspaceState::Moved:
+		return FName("Perforce.Branched_Small");
+	case EWorkspaceState::Deleted:
+		return FName("Perforce.MarkedForDelete_Small");
+	case EWorkspaceState::LocallyDeleted: // TODO: would need a dedicated icon
+	case EWorkspaceState::Conflicted: // TODO: would need a dedicated icon
+		return FName("Perforce.NotAtHeadRevision_Small");
+	case EWorkspaceState::LockedByOther:
+		return FName("Perforce.CheckedOutByOtherUser_Small");
+	case EWorkspaceState::Private: // Not controlled
+	case EWorkspaceState::Changed: // Changed but unchecked-out file is in a certain way not controlled
+		return FName("Perforce.NotInDepot_Small");
+	case EWorkspaceState::Unknown:
+	case EWorkspaceState::Ignored:
+	case EWorkspaceState::Controlled: // (Unchanged) same as "Pristine" for Perforce (not checked out) ie no icon
+	default:
+		return NAME_None;
+	}
+}
+
+FText FPlasticSourceControlState::GetDisplayName() const
+{
+	switch(WorkspaceState)
+	{
+	case EWorkspaceState::Unknown:
+		return LOCTEXT("Unknown", "Unknown");
+	case EWorkspaceState::Ignored:
+		return LOCTEXT("Ignored", "Ignored");
+	case EWorkspaceState::Controlled:
+		return LOCTEXT("Controlled", "Controlled");
+	case EWorkspaceState::CheckedOut:
+		return LOCTEXT("CheckedOut", "CheckedOut");
+	case EWorkspaceState::Added:
+		return LOCTEXT("Added", "Added");
+	case EWorkspaceState::Moved:
+		return LOCTEXT("Moved", "Moved");
+	case EWorkspaceState::Copied:
+		return LOCTEXT("Copied", "Copied");
+	case EWorkspaceState::Replaced:
+		return LOCTEXT("Replaced", "Replaced");
+	case EWorkspaceState::Deleted:
+		return LOCTEXT("Deleted", "Deleted");
+	case EWorkspaceState::LocallyDeleted:
+		return LOCTEXT("LocallyDeleted", "Missing");
+	case EWorkspaceState::Changed:
+		return LOCTEXT("Changed", "Changed");
+	case EWorkspaceState::Conflicted:
+		return LOCTEXT("ContentsConflict", "Contents Conflict");
+	case EWorkspaceState::LockedByOther:
+		return FText::Format(LOCTEXT("CheckedOutOther", "Checked out by: {0}"), FText::FromString(LockedBy));
+	case EWorkspaceState::Private:
+		return LOCTEXT("NotControlled", "Not Under Source Control");
+	}
+
+	return FText();
+}
+
+FText FPlasticSourceControlState::GetDisplayTooltip() const
+{
+	switch(WorkspaceState)
+	{
+	case EWorkspaceState::Unknown:
+		return LOCTEXT("Unknown_Tooltip", "Unknown source control state");
+	case EWorkspaceState::Ignored:
+		return LOCTEXT("Ignored_Tooltip", "Item is being ignored.");
+	case EWorkspaceState::Controlled:
+		return LOCTEXT("Pristine_Tooltip", "There are no modifications");
+	case EWorkspaceState::CheckedOut:
+		return LOCTEXT("CheckedOut_Tooltip", "Item is checked out");
+	case EWorkspaceState::Added:
+		return LOCTEXT("Added_Tooltip", "Item has been added");
+	case EWorkspaceState::Moved:
+		return LOCTEXT("Moved_Tooltip", "Item has been moved or renamed");
+	case EWorkspaceState::Copied:
+		return LOCTEXT("Copied_Tooltip", "Item has been copied");
+	case EWorkspaceState::Replaced:
+		return LOCTEXT("Replaced_Tooltip", "Item has been replaced / merged");
+	case EWorkspaceState::Deleted:
+		return LOCTEXT("Deleted_Tooltip", "Item is scheduled for deletion");
+	case EWorkspaceState::LocallyDeleted:
+		return LOCTEXT("LocallyDeleted_Tooltip", "Item is missing");
+	case EWorkspaceState::Changed:
+		return LOCTEXT("Modified_Tooltip", "Item has been modified");
+	case EWorkspaceState::Conflicted:
+		return LOCTEXT("ContentsConflict_Tooltip", "The contents of the item conflict with updates received from the repository.");
+	case EWorkspaceState::Private:
+		return LOCTEXT("NotControlled_Tooltip", "Item is not under version control.");
+	}
+
+	return FText();
+}
+
+const FString& FPlasticSourceControlState::GetFilename() const
+{
+	// UE_LOG(LogSourceControl, Log, TEXT("GetFilename(%s)"), *LocalFilename);
+
+	return LocalFilename;
+}
+
+const FDateTime& FPlasticSourceControlState::GetTimeStamp() const
+{
+	// UE_LOG(LogSourceControl, Log, TEXT("GetTimeStamp(%s)=%s"), *LocalFilename, *TimeStamp.ToString());
+
+	return TimeStamp;
+}
+
+// Deleted and Missing assets cannot appear in the Content Browser but does appear in Submit to Source Control Window
+bool FPlasticSourceControlState::CanCheckIn() const
+{
+	const bool bCanCheckIn = WorkspaceState == EWorkspaceState::Added
+		|| WorkspaceState == EWorkspaceState::Deleted
+		|| WorkspaceState == EWorkspaceState::LocallyDeleted
+		|| WorkspaceState == EWorkspaceState::Changed // NOTE: Comment to enable checkout on prompt of a Changed file (see bellow)
+		|| WorkspaceState == EWorkspaceState::Moved
+		|| WorkspaceState == EWorkspaceState::Copied
+		|| WorkspaceState == EWorkspaceState::Replaced
+		|| WorkspaceState == EWorkspaceState::CheckedOut;
+
+	UE_LOG(LogSourceControl, Log, TEXT("%s CanCheckIn=%d"), *LocalFilename, bCanCheckIn);
+
+	return bCanCheckIn;
+}
+
+bool FPlasticSourceControlState::CanCheckout() const
+{
+	const bool bCanCheckout  = (   WorkspaceState == EWorkspaceState::Controlled	// In source control, Unmodified
+								|| WorkspaceState == EWorkspaceState::Changed		// In source control, but not checked-out
+								|| WorkspaceState == EWorkspaceState::Replaced)		// In source control, merged, waiting for checkin to conclude the merge 
+								&& IsCurrent(); // Is up to date (at the revision of the repo)
+
+	UE_LOG(LogSourceControl, Log, TEXT("%s CanCheckout=%d"), *LocalFilename, bCanCheckout);
+
+	return bCanCheckout;
+}
+
+bool FPlasticSourceControlState::IsCheckedOut() const
+{
+	const bool bIsCheckedOut = WorkspaceState == EWorkspaceState::CheckedOut
+							|| WorkspaceState == EWorkspaceState::Moved
+							|| WorkspaceState == EWorkspaceState::Conflicted
+							|| WorkspaceState == EWorkspaceState::Replaced // Workaround to enable checkin
+							|| WorkspaceState == EWorkspaceState::Changed; // Workaround to enable checkin  NOTE: Comment to enable checkout on prompt of a Changed file (see above)
+
+	if (bIsCheckedOut) UE_LOG(LogSourceControl, Log, TEXT("%s IsCheckedOut"), *LocalFilename);
+
+	return bIsCheckedOut;
+}
+
+bool FPlasticSourceControlState::IsCheckedOutOther(FString* Who) const
+{
+	if (Who != NULL)
+	{
+		*Who = LockedBy;
+	}
+	const bool bIsLockedByOther = WorkspaceState == EWorkspaceState::LockedByOther;
+
+	// @todo: temporary debug log
+	// if (bIsLockedByOther) UE_LOG(LogSourceControl, Log, TEXT("%s IsCheckedOutOther by '%s' (%s)"), *LocalFilename, *LockedBy, *LockedWhere);
+
+	return bIsLockedByOther;
+}
+
+
+/** Get whether this file is checked out in a different branch, if no branch is specified defaults to FEngineVerion current branch */
+bool FPlasticSourceControlState::IsCheckedOutInOtherBranch(const FString& CurrentBranch) const
+{
+	return false;
+}
+
+/** Get whether this file is modified in a different branch, if no branch is specified defaults to FEngineVerion current branch */
+bool FPlasticSourceControlState::IsModifiedInOtherBranch(const FString& CurrentBranch) const
+{
+	return false;
+}
+
+/** Get whether this file is checked out or modified in a different branch, if no branch is specified defaults to FEngineVerion current branch */
+bool FPlasticSourceControlState::IsCheckedOutOrModifiedInOtherBranch(const FString& CurrentBranch) const
+{
+	return false;
+}
+
+/** Get the other branches this file is checked out in */
+TArray<FString> FPlasticSourceControlState::GetCheckedOutBranches() const
+{
+	return TArray<FString>();
+}
+
+/** Get the user info for checkouts on other branches */
+FString FPlasticSourceControlState::GetOtherUserBranchCheckedOuts() const
+{
+	return FString();
+}
+
+/** Get head modification information for other branches
+ * @returns true with parameters populated if there is a branch with a newer modification (edit/delete/etc)
+*/
+bool FPlasticSourceControlState::GetOtherBranchHeadModification(FString& HeadBranchOut, FString& ActionOut, int32& HeadChangeListOut) const
+{
+	return false;
+}
+
+
+bool FPlasticSourceControlState::IsCurrent() const
+{
+	// NOTE: Deleted assets get a "-1" HeadRevision which we do not want to override the real icon state
+	const bool bIsCurrent = (LocalRevisionChangeset == DepotRevisionChangeset) || (WorkspaceState == EWorkspaceState::Deleted);
+
+	// @todo: temporary debug log
+	// if (bIsCurrent) UE_LOG(LogSourceControl, Log, TEXT("%s IsCurrent"), *LocalFilename, );
+	
+	return bIsCurrent;
+}
+
+bool FPlasticSourceControlState::IsSourceControlled() const
+{
+	const bool bIsSourceControlled = WorkspaceState != EWorkspaceState::Private
+								  && WorkspaceState != EWorkspaceState::Ignored
+								  && WorkspaceState != EWorkspaceState::Unknown;
+
+	if (!bIsSourceControlled && !IsUnknown()) UE_LOG(LogSourceControl, Log, TEXT("%s NOT SourceControlled"), *LocalFilename);
+
+	return bIsSourceControlled;
+}
+
+bool FPlasticSourceControlState::IsAdded() const
+{
+	const bool bIsAdded =	WorkspaceState == EWorkspaceState::Added
+						 || WorkspaceState == EWorkspaceState::Copied;
+
+	if (bIsAdded) UE_LOG(LogSourceControl, Log, TEXT("%s IsAdded"), *LocalFilename);
+
+	return bIsAdded;
+}
+
+bool FPlasticSourceControlState::IsDeleted() const
+{
+	const bool bIsDeleted =	WorkspaceState == EWorkspaceState::Deleted
+						 || WorkspaceState == EWorkspaceState::LocallyDeleted;
+
+	if (bIsDeleted) UE_LOG(LogSourceControl, Log, TEXT("%s IsDeleted"), *LocalFilename);
+
+	return bIsDeleted;
+}
+
+bool FPlasticSourceControlState::IsIgnored() const
+{
+	if (WorkspaceState == EWorkspaceState::Ignored) UE_LOG(LogSourceControl, Log, TEXT("%s IsIgnored"), *LocalFilename);
+
+	return WorkspaceState == EWorkspaceState::Ignored;
+}
+
+bool FPlasticSourceControlState::CanEdit() const
+{
+	const bool bCanEdit =  WorkspaceState == EWorkspaceState::CheckedOut
+						|| WorkspaceState == EWorkspaceState::Added
+						|| WorkspaceState == EWorkspaceState::Moved
+						|| WorkspaceState == EWorkspaceState::Copied
+						|| WorkspaceState == EWorkspaceState::Replaced;
+
+	UE_LOG(LogSourceControl, Log, TEXT("%s CanEdit=%d"), *LocalFilename, bCanEdit);
+
+	return bCanEdit;
+}
+
+bool FPlasticSourceControlState::CanDelete() const
+{
+	return !IsCheckedOutOther() && IsSourceControlled() && IsCurrent();
+}
+
+bool FPlasticSourceControlState::IsUnknown() const
+{
+	return WorkspaceState == EWorkspaceState::Unknown;
+}
+
+bool FPlasticSourceControlState::IsModified() const
+{
+	// Warning: for a clean "check-in" (commit) checked-out files unmodified should be removed from the changeset (the index)
+	//
+	// Thus, before check-in UE4 Editor call RevertUnchangedFiles() in PromptForCheckin() and CheckinFiles().
+	//
+	// So here we must take care to enumerate all states that need to be commited, all other will be discarded:
+	//  - Unknown
+	//  - Controlled (Unchanged)
+	//  - Private (Not Controlled)
+	//  - Ignored
+	const bool bIsModified =   WorkspaceState == EWorkspaceState::CheckedOut
+							|| WorkspaceState == EWorkspaceState::Added
+							|| WorkspaceState == EWorkspaceState::Moved
+							|| WorkspaceState == EWorkspaceState::Copied
+							|| WorkspaceState == EWorkspaceState::Replaced
+							|| WorkspaceState == EWorkspaceState::Deleted
+							|| WorkspaceState == EWorkspaceState::LocallyDeleted
+							|| WorkspaceState == EWorkspaceState::Changed
+							|| WorkspaceState == EWorkspaceState::Conflicted;
+
+	UE_LOG(LogSourceControl, Log, TEXT("%s IsModified=%d"), *LocalFilename, bIsModified);
+
+	return bIsModified;
+}
+
+
+bool FPlasticSourceControlState::CanAdd() const
+{
+	if (!IsUnknown()) UE_LOG(LogSourceControl, Log, TEXT("%s CanAdd=%d"), *LocalFilename, WorkspaceState == EWorkspaceState::Private);
+
+	return WorkspaceState == EWorkspaceState::Private;
+}
+
+bool FPlasticSourceControlState::IsConflicted() const
+{
+	if (WorkspaceState == EWorkspaceState::Conflicted) UE_LOG(LogSourceControl, Log, TEXT("%s IsConflicted"), *LocalFilename);
+
+	return WorkspaceState == EWorkspaceState::Conflicted;
+}
+
+bool FPlasticSourceControlState::CanRevert() const
+{
+	return IsModified();
+}
+
+#undef LOCTEXT_NAMESPACE
diff --git a/Engine/Plugins/Developer/PlasticSourceControl/Source/PlasticSourceControl/Private/PlasticSourceControlState.h b/Engine/Plugins/Developer/PlasticSourceControl/Source/PlasticSourceControl/Private/PlasticSourceControlState.h
new file mode 100644
--- /dev/null
+++ b/Engine/Plugins/Developer/PlasticSourceControl/Source/PlasticSourceControl/Private/PlasticSourceControlState.h
@@ -0,0 +1,170 @@
+// Copyright 1998-2019 Epic Games, Inc. All Rights Reserved.
+
+#pragma once
+
+#include "CoreMinimal.h"
+#include "ISourceControlState.h"
+#include "PlasticSourceControlRevision.h"
+
+namespace EWorkspaceState
+{
+	enum Type
+	{
+		Unknown,
+		Ignored,
+		Controlled, // called "Pristine" in Perforce, "Unchanged" in Git, "Clean" in SVN
+		CheckedOut,
+		Added,
+		Moved, // Renamed
+		Copied,
+		Replaced, // Replaced / Merged
+		Deleted,
+		LocallyDeleted, // Missing
+		Changed, // Modified but not CheckedOut
+		Conflicted,
+		LockedByOther, // LockedBy with name of someone else than cm whoami
+		Private, // "Not Controlled"/"Not In Depot"/"Untracked"
+	};
+
+	// debug log utility
+	const TCHAR* ToString(EWorkspaceState::Type InWorkspaceState);
+}
+
+class FPlasticSourceControlState : public ISourceControlState, public TSharedFromThis<FPlasticSourceControlState, ESPMode::ThreadSafe>
+{
+public:
+
+	FPlasticSourceControlState(const FString& InLocalFilename)
+		: LocalFilename(InLocalFilename)
+		, WorkspaceState(EWorkspaceState::Unknown)
+		, DepotRevisionChangeset(-1)
+		, LocalRevisionChangeset(-1)
+		, TimeStamp(0)
+	{
+	}
+
+	FPlasticSourceControlState(const FPlasticSourceControlState& InState) = delete;
+	const FPlasticSourceControlState& operator=(const FPlasticSourceControlState& InState) = delete;
+
+	FPlasticSourceControlState(FPlasticSourceControlState&& InState)
+	{
+		Move(MoveTemp(InState));
+	}
+
+	const FPlasticSourceControlState& operator=(FPlasticSourceControlState&& InState)
+	{
+		Move(MoveTemp(InState));
+		return *this;
+	}
+
+	void Move(FPlasticSourceControlState&& InState)
+	{
+		History = MoveTemp(InState.History);
+		LocalFilename = MoveTemp(InState.LocalFilename);
+		RepSpec = MoveTemp(InState.RepSpec);
+		PendingMergeFilename = MoveTemp(InState.PendingMergeFilename);
+		PendingMergeBaseChangeset = InState.PendingMergeBaseChangeset;
+		PendingMergeSourceChangeset = InState.PendingMergeSourceChangeset;
+		PendingMergeParameters = MoveTemp(InState.PendingMergeParameters);
+		LockedBy = MoveTemp(InState.LockedBy);
+		LockedWhere = MoveTemp(InState.LockedWhere);
+		WorkspaceState = InState.WorkspaceState;
+		DepotRevisionChangeset = InState.DepotRevisionChangeset;
+		LocalRevisionChangeset = InState.LocalRevisionChangeset;
+		MovedFrom = MoveTemp(InState.MovedFrom);
+		TimeStamp = InState.TimeStamp;
+
+		// Update Revision's pointer to the new State instance
+		for (const auto& Revision : History)
+		{
+			Revision->State = this;
+		}
+	}
+
+	// debug log utility
+	const TCHAR* ToString() const
+	{
+		return EWorkspaceState::ToString(WorkspaceState);
+	}
+
+	/** ISourceControlState interface */
+	virtual int32 GetHistorySize() const override;
+	virtual TSharedPtr<class ISourceControlRevision, ESPMode::ThreadSafe> GetHistoryItem(int32 HistoryIndex) const override;
+	virtual TSharedPtr<class ISourceControlRevision, ESPMode::ThreadSafe> FindHistoryRevision(int32 RevisionNumber) const override;
+	virtual TSharedPtr<class ISourceControlRevision, ESPMode::ThreadSafe> FindHistoryRevision(const FString& InRevision) const override;
+	virtual TSharedPtr<class ISourceControlRevision, ESPMode::ThreadSafe> GetBaseRevForMerge() const override;
+	virtual FName GetIconName() const override;
+	virtual FName GetSmallIconName() const override;
+	virtual FText GetDisplayName() const override;
+	virtual FText GetDisplayTooltip() const override;
+	virtual const FString& GetFilename() const override;
+	virtual const FDateTime& GetTimeStamp() const override;
+	virtual bool CanCheckIn() const override;
+	virtual bool CanCheckout() const override;
+	virtual bool IsCheckedOut() const override;
+	virtual bool IsCheckedOutOther(FString* Who = nullptr) const override;
+	virtual bool IsCheckedOutInOtherBranch(const FString& CurrentBranch = FString()) const /* override UE4.20 */;
+	virtual bool IsModifiedInOtherBranch(const FString& CurrentBranch = FString()) const /* override UE4.20 */;
+	virtual bool IsCheckedOutOrModifiedInOtherBranch(const FString& CurrentBranch = FString()) const /* override UE4.20 */;
+	virtual TArray<FString> GetCheckedOutBranches() const /* override UE4.20 */;
+	virtual FString GetOtherUserBranchCheckedOuts() const /* override UE4.20 */;
+	virtual bool GetOtherBranchHeadModification(FString& HeadBranchOut, FString& ActionOut, int32& HeadChangeListOut) const /* override UE4.20 */;
+	virtual bool IsCurrent() const override;
+	virtual bool IsSourceControlled() const override;
+	virtual bool IsAdded() const override;
+	virtual bool IsDeleted() const override;
+	virtual bool IsIgnored() const override;
+	virtual bool CanEdit() const override;
+	virtual bool IsUnknown() const override;
+	virtual bool IsModified() const override;
+	virtual bool CanAdd() const override;
+	virtual bool CanDelete() const override;
+	virtual bool IsConflicted() const override;
+	virtual bool CanRevert() const; // override; TODO UE 4.19 only - commented for backward compatibility
+
+public:
+	/** History of the item, if any */
+	TPlasticSourceControlHistory History;
+
+	/** Filename on disk */
+	FString LocalFilename;
+
+	/** Depot and Server info (in the form repo@server:port) */
+	FString RepSpec;
+
+	/** Relative filename of the file in merge conflict */
+	FString PendingMergeFilename;
+
+	/** Changeset with which our local revision diverged from the source/remote revision */
+	int32 PendingMergeBaseChangeset;
+
+	/** Changeset of the source/remote revision of the merge in progress */
+	int32 PendingMergeSourceChangeset;
+
+	/** Plastic SCM Parameters of the merge in progress */
+	TArray<FString> PendingMergeParameters;
+
+	/** If a user (another or ourself) has this file locked, this contains his name. */
+	FString LockedBy;
+
+	/** Location of the locked file. */
+	FString LockedWhere;
+
+	/** State of the workspace */
+	EWorkspaceState::Type WorkspaceState;
+
+	/** Latest revision number of the file in the depot */
+	int DepotRevisionChangeset;
+
+	/** Latest revision number at which a file was synced to before being edited */
+	int LocalRevisionChangeset;
+
+	/** Original name in case of a Moved/Renamed file */
+	FString MovedFrom;
+
+	/** Whether the file is a binary file or not */
+//	bool bBinary;
+
+	/** The timestamp of the last update */
+	FDateTime TimeStamp;
+};
diff --git a/Engine/Plugins/Developer/PlasticSourceControl/Source/PlasticSourceControl/Private/PlasticSourceControlUtils.cpp b/Engine/Plugins/Developer/PlasticSourceControl/Source/PlasticSourceControl/Private/PlasticSourceControlUtils.cpp
new file mode 100644
--- /dev/null
+++ b/Engine/Plugins/Developer/PlasticSourceControl/Source/PlasticSourceControl/Private/PlasticSourceControlUtils.cpp
@@ -0,0 +1,1535 @@
+// Copyright 1998-2019 Epic Games, Inc. All Rights Reserved.
+
+#include "PlasticSourceControlUtils.h"
+#include "PlasticSourceControlCommand.h"
+#include "PlasticSourceControlModule.h"
+#include "PlasticSourceControlSettings.h"
+#include "PlasticSourceControlState.h"
+#include "HAL/PlatformFilemanager.h"
+#include "HAL/PlatformProcess.h"
+#include "HAL/FileManager.h"
+#include "Misc/ScopeLock.h"
+#include "Misc/FileHelper.h"
+#include "Misc/Paths.h"
+#include "HAL/PlatformTime.h"
+#include "Modules/ModuleManager.h"
+#include "XmlParser.h"
+
+#if PLATFORM_LINUX
+#include <sys/ioctl.h>
+#endif
+
+#if PLATFORM_WINDOWS
+#include "Windows/WindowsHWrapper.h" // SECURITY_ATTRIBUTES
+#undef GetUserName
+#endif
+
+
+namespace PlasticSourceControlConstants
+{
+#if PLATFORM_WINDOWS
+	const TCHAR* pchDelim = TEXT("\r\n");
+#else
+	const TCHAR* pchDelim = TEXT("\n");
+#endif
+}
+
+FScopedTempFile::FScopedTempFile(const FText& InText)
+{
+	Filename = FPaths::CreateTempFilename(*FPaths::ProjectLogDir(), TEXT("Plastic-Temp"), TEXT(".txt"));
+	if(!FFileHelper::SaveStringToFile(InText.ToString(), *Filename, FFileHelper::EEncodingOptions::ForceUTF8WithoutBOM))
+	{
+		UE_LOG(LogSourceControl, Error, TEXT("Failed to write to temp file: %s"), *Filename);
+	}
+}
+
+FScopedTempFile::~FScopedTempFile()
+{
+	if(FPaths::FileExists(Filename))
+	{
+		if(!FPlatformFileManager::Get().GetPlatformFile().DeleteFile(*Filename))
+		{
+			UE_LOG(LogSourceControl, Error, TEXT("Failed to delete temp file: %s"), *Filename);
+		}
+	}
+}
+
+const FString& FScopedTempFile::GetFilename() const
+{
+	return Filename;
+}
+
+// Needed to SetHandleInformation() on WritePipe for input (opposite of ReadPipe, for output) (idem FInteractiveProcess)
+static FORCEINLINE bool CreatePipeWrite(void*& ReadPipe, void*& WritePipe)
+{
+#if PLATFORM_WINDOWS
+	SECURITY_ATTRIBUTES Attr = { sizeof(SECURITY_ATTRIBUTES), NULL, true };
+
+	if (!::CreatePipe(&ReadPipe, &WritePipe, &Attr, 0))
+	{
+		return false;
+	}
+
+	if (!::SetHandleInformation(WritePipe, HANDLE_FLAG_INHERIT, 0))
+	{
+		return false;
+	}
+
+	return true;
+#else
+	return FPlatformProcess::CreatePipe(ReadPipe, WritePipe);
+#endif // PLATFORM_WINDOWS
+}
+
+namespace PlasticSourceControlUtils
+{
+// Command-line interface paramters and output format changed with version 8.0.16.3000
+// For more details, see https://www.plasticscm.com/download/releasenotes/oldernotes?release=8.0.16.3000
+static bool				bIsNewVersion80163000 = false;
+
+// In/Out Pipes for the 'cm shell' persistent process
+static void*			ShellOutputPipeRead = nullptr;
+static void*			ShellOutputPipeWrite = nullptr;
+static void*			ShellInputPipeRead = nullptr;
+static void*			ShellInputPipeWrite = nullptr;
+static FProcHandle		ShellProcessHandle;
+static FCriticalSection	ShellCriticalSection;
+static size_t			ShellCommandCounter = -1;
+
+// Internal function to cleanup (called under the critical section)
+static void _CleanupBackgroundCommandLineShell()
+{
+	FPlatformProcess::ClosePipe(ShellInputPipeRead, ShellInputPipeWrite);
+	FPlatformProcess::ClosePipe(ShellOutputPipeRead, ShellOutputPipeWrite);
+	ShellOutputPipeRead = ShellOutputPipeWrite = nullptr;
+	ShellInputPipeRead = ShellInputPipeWrite = nullptr;
+}
+
+// Internal function to actualy launch the Plastic SCM background 'cm shell' process if possible (called under the critical section)
+static bool _StartBackgroundPlasticShell(const FString& InPathToPlasticBinary, const FString& InWorkingDirectory)
+{
+	const FString FullCommand(TEXT("shell"));
+
+	const bool bLaunchDetached = false;				// the new process will NOT have its own window
+	const bool bLaunchHidden = true;				// the new process will be minimized in the task bar
+	const bool bLaunchReallyHidden = bLaunchHidden; // the new process will not have a window or be in the task bar
+
+	verify(FPlatformProcess::CreatePipe(ShellOutputPipeRead, ShellOutputPipeWrite));	// For reading from child process
+	verify(CreatePipeWrite(ShellInputPipeRead, ShellInputPipeWrite));	// For writing to child process
+
+	ShellProcessHandle = FPlatformProcess::CreateProc(*InPathToPlasticBinary, *FullCommand, bLaunchDetached, bLaunchHidden, bLaunchReallyHidden, nullptr, 0, *InWorkingDirectory, ShellOutputPipeWrite, ShellInputPipeRead);
+	if (!ShellProcessHandle.IsValid())
+	{
+		UE_LOG(LogSourceControl, Warning, TEXT("Failed to launch 'cm shell'")); // not a bug, just no Plastic SCM cli found
+		_CleanupBackgroundCommandLineShell();
+	}
+	else
+	{
+		UE_LOG(LogSourceControl, Log, TEXT("LaunchBackgroundPlasticShell: '%s %s' ok (handle %d)"), *InPathToPlasticBinary, *FullCommand, ShellProcessHandle.Get());
+		ShellCommandCounter = 0;
+	}
+
+	return ShellProcessHandle.IsValid();
+}
+
+// Internal function (called under the critical section)
+static void _RestartBackgroundCommandLineShell()
+{
+	const FPlasticSourceControlModule& PlasticSourceControl = FModuleManager::GetModuleChecked<FPlasticSourceControlModule>("PlasticSourceControl");
+	const FString& PathToPlasticBinary = PlasticSourceControl.AccessSettings().GetBinaryPath();
+	const FString& WorkingDirectory = PlasticSourceControl.GetProvider().GetPathToWorkspaceRoot();
+
+	FPlatformProcess::CloseProc(ShellProcessHandle);
+	_CleanupBackgroundCommandLineShell();
+	_StartBackgroundPlasticShell(PathToPlasticBinary, WorkingDirectory);
+}
+
+// Internal function (called under the critical section)
+static bool _RunCommandInternal(const FString& InCommand, const TArray<FString>& InParameters, const TArray<FString>& InFiles, const EConcurrency::Type InConcurrency, FString& OutResults, FString& OutErrors)
+{
+	bool bResult = false;
+
+	// Detect previous crash of cm.exe and restart 'cm shell'
+	if (!FPlatformProcess::IsProcRunning(ShellProcessHandle))
+	{
+		UE_LOG(LogSourceControl, Warning, TEXT("RunCommand(%d): 'cm shell' has stopped. Restarting!"), ShellCommandCounter);
+		_RestartBackgroundCommandLineShell();
+	}
+
+	ShellCommandCounter++;
+
+	// Start with the Plastic command itself ("status", "log", "chekin"...)
+	FString FullCommand = InCommand;
+	// Append to the command all parameters, and then finally the files
+	for (const FString& Parameter : InParameters)
+	{
+		FullCommand += TEXT(" ");
+		FullCommand += Parameter;
+	}
+	for (const FString& File : InFiles)
+	{
+		FullCommand += TEXT(" \"");
+		FullCommand += File;
+		FullCommand += TEXT("\"");
+	}
+	// @todo: temporary debug logs (without the end-of-line)
+	const FString LoggableCommand = FullCommand.Left(256); // Limit command log size to 256 characters
+	UE_LOG(LogSourceControl, Log, TEXT("RunCommand(%d): '%s' (%d chars, %d files)"), ShellCommandCounter, *LoggableCommand, FullCommand.Len()+1, InFiles.Num());
+	FullCommand += TEXT('\n'); // Finalize the command line
+
+	// Send command to 'cm shell' process
+	const bool bWriteOk = FPlatformProcess::WritePipe(ShellInputPipeWrite, FullCommand);
+
+	// And wait up to 180.0 seconds for any kind of output from cm shell: in case of lengthier operation, intermediate output (like percentage of progress) is expected, which would refresh the timout
+	const double Timeout = 180.0;
+	const double StartTimestamp = FPlatformTime::Seconds();
+	double LastActivity = StartTimestamp;
+	double LastLog = StartTimestamp;
+	const double LogInterval = 5.0;
+	int32 PreviousLogLen = 0;
+	while (FPlatformProcess::IsProcRunning(ShellProcessHandle))
+	{
+		FString Output = FPlatformProcess::ReadPipe(ShellOutputPipeRead);
+		if (0 < Output.Len())
+		{
+			LastActivity = FPlatformTime::Seconds(); // freshen the timestamp while cm is still actively outputting information
+			OutResults.Append(MoveTemp(Output));
+			// Search the output for the line containing the result code, also indicating the end of the command
+			const uint32 IndexCommandResult = OutResults.Find(TEXT("CommandResult "), ESearchCase::CaseSensitive, ESearchDir::FromEnd);
+			if (INDEX_NONE != IndexCommandResult)
+			{
+				const uint32 IndexEndResult = OutResults.Find(PlasticSourceControlConstants::pchDelim, ESearchCase::CaseSensitive, ESearchDir::FromStart, IndexCommandResult + 14);
+				if (INDEX_NONE != IndexEndResult)
+				{
+					const FString Result = OutResults.Mid(IndexCommandResult + 14, IndexEndResult - IndexCommandResult - 14);
+					const int32 ResultCode = FCString::Atoi(*Result);
+					bResult = (0 == ResultCode);
+					// remove the CommandResult line from the OutResults
+					OutResults.RemoveAt(IndexCommandResult, OutResults.Len() - IndexCommandResult);
+					break;
+				}
+			}
+		}
+		else if ((FPlatformTime::Seconds() - LastLog > LogInterval) && (PreviousLogLen < OutResults.Len()) && (InConcurrency == EConcurrency::Asynchronous))
+		{
+			// In case of long running operation, start to print intermediate output from cm shell (like percentage of progress)
+			// (but only when running Asynchronous commands, since Synchronous commands block the main thread until they finish)
+			UE_LOG(LogSourceControl, Log, TEXT("RunCommand(%d): '%s' in progress for %lfs...\n%s"), ShellCommandCounter, *InCommand, (FPlatformTime::Seconds() - StartTimestamp), *OutResults.Mid(PreviousLogLen));
+			PreviousLogLen = OutResults.Len();
+			LastLog = FPlatformTime::Seconds(); // freshen the timestamp of last log
+		}
+		else if (FPlatformTime::Seconds() - LastActivity > Timeout)
+		{
+			// In case of timeout, ask the blocking 'cm shell' process to exit, and detach from it immediatly: it will be relaunched by next command
+			UE_LOG(LogSourceControl, Error, TEXT("RunCommand(%d): '%s' %d TIMEOUT after %lfs output:\n%s"), ShellCommandCounter, *InCommand, bResult, (FPlatformTime::Seconds() - StartTimestamp), *OutResults.Mid(PreviousLogLen));
+			FPlatformProcess::WritePipe(ShellInputPipeWrite, TEXT("exit"));
+			FPlatformProcess::CloseProc(ShellProcessHandle);
+			_CleanupBackgroundCommandLineShell();
+		}
+
+		FPlatformProcess::Sleep(0.001f);
+	}
+	if (!InCommand.Equals(TEXT("exit")))
+	{
+		if (!FPlatformProcess::IsProcRunning(ShellProcessHandle))
+		{
+			// 'cm shell' normally only terminates in case of 'exit' command. Will restart on next command.
+			UE_LOG(LogSourceControl, Error, TEXT("RunCommand(%d): '%s' 'cm shell' stopped after %lfs output:\n%s"), ShellCommandCounter, *LoggableCommand, (FPlatformTime::Seconds() - StartTimestamp), *OutResults.Left(4096)); // Limit result size to 4096 characters
+		}
+		else if (!bResult)
+		{
+			UE_LOG(LogSourceControl, Warning, TEXT("RunCommand(%d): '%s' (in %lfs) output:\n%s"), ShellCommandCounter, *LoggableCommand, (FPlatformTime::Seconds() - StartTimestamp), *OutResults.Left(4096)); // Limit result size to 4096 characters
+		}
+		else
+		{
+			if (PreviousLogLen > 0)
+			{
+				UE_LOG(LogSourceControl, Log, TEXT("RunCommand(%d): '%s' (in %lfs) output:\n%s"), ShellCommandCounter, *LoggableCommand, (FPlatformTime::Seconds() - StartTimestamp), *OutResults.Mid(PreviousLogLen).Left(4096)); // Limit result size to 4096 characters
+			}
+			else
+			{
+				UE_LOG(LogSourceControl, Log, TEXT("RunCommand(%d): '%s' (in %lfs) output:\n%s"), ShellCommandCounter, *LoggableCommand, (FPlatformTime::Seconds() - StartTimestamp), *OutResults.Left(4096)); // Limit result size to 4096 characters
+			}
+		}
+	}
+	// Return output as error if result code is an error
+	if (!bResult)
+	{
+		OutErrors = MoveTemp(OutResults);
+	}
+
+	return bResult;
+}
+
+// Internal function (called under the critical section)
+static void _ExitBackgroundCommandLineShell()
+{
+	// Tell the 'cm shell' to exit
+	FString Results, Errors;
+	_RunCommandInternal(TEXT("exit"), TArray<FString>(), TArray<FString>(), EConcurrency::Synchronous, Results, Errors);
+	// And wait up to one seconde for its termination
+	int timeout = 100;
+	while (FPlatformProcess::IsProcRunning(ShellProcessHandle) && (0 < timeout--))
+	{
+		FPlatformProcess::Sleep(0.01f);
+	}
+	FPlatformProcess::CloseProc(ShellProcessHandle);
+	_CleanupBackgroundCommandLineShell();
+}
+
+// Launch the Plastic SCM background 'cm shell' process in background for optimized successive commands (thread-safe)
+bool LaunchBackgroundPlasticShell(const FString& InPathToPlasticBinary, const FString& InWorkingDirectory)
+{
+	// Protect public APIs from multi-thread access
+	FScopeLock Lock(&ShellCriticalSection);
+
+	// terminate previous shell if one is already running
+	if (ShellProcessHandle.IsValid())
+	{
+		_ExitBackgroundCommandLineShell();
+	}
+
+	return _StartBackgroundPlasticShell(InPathToPlasticBinary, InWorkingDirectory);
+}
+
+// Terminate the background 'cm shell' process and associated pipes (thread-safe)
+void Terminate()
+{
+	// Protect public APIs from multi-thread access
+	FScopeLock Lock(&ShellCriticalSection);
+
+	if (ShellProcessHandle.IsValid())
+	{
+		_ExitBackgroundCommandLineShell();
+	}
+}
+
+// Run command (thread-safe)
+static bool RunCommandInternal(const FString& InCommand, const TArray<FString>& InParameters, const TArray<FString>& InFiles, const EConcurrency::Type InConcurrency, FString& OutResults, FString& OutErrors)
+{
+	bool bResult = false;
+
+	// Protect public APIs from multi-thread access
+	FScopeLock Lock(&ShellCriticalSection);
+
+	if (ShellProcessHandle.IsValid())
+	{
+		bResult = _RunCommandInternal(InCommand, InParameters, InFiles, InConcurrency, OutResults, OutErrors);
+	}
+	else
+	{
+		UE_LOG(LogSourceControl, Error, TEXT("RunCommand(%d): '%s': cm shell not running"), ShellCommandCounter, *InCommand);
+		OutErrors = InCommand + ": Plastic SCM shell not running!";
+	}
+
+	return bResult;
+}
+
+// Basic parsing or results & errors from the Plastic command line process
+bool RunCommand(const FString& InCommand, const TArray<FString>& InParameters, const TArray<FString>& InFiles, const EConcurrency::Type InConcurrency, TArray<FString>& OutResults, TArray<FString>& OutErrorMessages)
+{
+	FString Results;
+	FString Errors;
+
+	const bool bResult = RunCommandInternal(InCommand, InParameters, InFiles, InConcurrency, Results, Errors);
+
+	Results.ParseIntoArray(OutResults, PlasticSourceControlConstants::pchDelim, true);
+	Errors.ParseIntoArray(OutErrorMessages, PlasticSourceControlConstants::pchDelim, true);
+
+	return bResult;
+}
+
+FString FindPlasticBinaryPath()
+{
+#if PLATFORM_WINDOWS
+	return FString(TEXT("cm"));
+#else
+	return FString(TEXT("/usr/bin/cm"));
+#endif
+}
+
+// Find the root of the Plastic workspace, looking from the provided path and upward in its parent directories.
+bool FindRootDirectory(const FString& InPath, FString& OutWorkspaceRoot)
+{
+	bool bFound = false;
+	FString PathToPlasticSubdirectory;
+	OutWorkspaceRoot = InPath;
+
+	auto TrimTrailing = [](FString& Str, const TCHAR Char)
+	{
+		int32 Len = Str.Len();
+		while(Len && Str[Len - 1] == Char)
+		{
+			Str = Str.LeftChop(1);
+			Len = Str.Len();
+		}
+	};
+
+	TrimTrailing(OutWorkspaceRoot, '\\');
+	TrimTrailing(OutWorkspaceRoot, '/');
+
+	while(!bFound && !OutWorkspaceRoot.IsEmpty())
+	{
+		// Look for the ".plastic" subdirectory present at the root of every Plastic workspace
+		PathToPlasticSubdirectory = OutWorkspaceRoot / TEXT(".plastic");
+		bFound = IFileManager::Get().DirectoryExists(*PathToPlasticSubdirectory);
+		if(!bFound)
+		{
+			int32 LastSlashIndex;
+			if(OutWorkspaceRoot.FindLastChar(TEXT('/'), LastSlashIndex))
+			{
+				OutWorkspaceRoot = OutWorkspaceRoot.Left(LastSlashIndex);
+			}
+			else
+			{
+				OutWorkspaceRoot.Empty();
+			}
+		}
+	}
+	if (!bFound)
+	{
+		OutWorkspaceRoot = InPath; // If not found, return the provided dir as best possible root.
+	}
+	return bFound;
+}
+
+/**
+ * @brief Compare Plastic SCM cli version strings.
+ * @param VersionA		PlasticSCM version string in the form "0.0.0.0" (as returned by GetPlasticScmVersion)
+ * @param VersionB		PlasticSCM version string in the form "0.0.0.0" (as returned by GetPlasticScmVersion)
+ * @returns true if VersionA is lower than VersionB
+*/
+static bool PlasticScmVersionLess(const FString& VersionA, const FString& VersionB) {
+	struct Version {
+		Version(const FString& V) {
+			TArray<FString> Parts;
+			const int32 N = V.ParseIntoArray(Parts, TEXT("."));
+			check(N == 4);
+			a = FCString::Atoi(*Parts[0]);
+			b = FCString::Atoi(*Parts[1]);
+			c = FCString::Atoi(*Parts[2]);
+			d = FCString::Atoi(*Parts[3]);
+		}
+		int a, b, c, d;
+	} A(VersionA), B(VersionB);
+	if (A.a < B.a) return true;
+	if (B.a < A.a) return false;
+	if (A.b < B.b) return true;
+	if (B.b < A.b) return false;
+	if (A.c < B.c) return true;
+	if (B.c < A.c) return false;
+	if (A.d < B.d) return true;
+	if (B.d < A.d) return false;
+	return false; // Equal
+}
+
+// This is called once by FPlasticSourceControlProvider::CheckPlasticAvailability()
+void GetPlasticScmVersion(FString& OutPlasticScmVersion)
+{
+	TArray<FString> InfoMessages;
+	TArray<FString> ErrorMessages;
+	const bool bResult = RunCommand(TEXT("version"), TArray<FString>(), TArray<FString>(), EConcurrency::Synchronous, InfoMessages, ErrorMessages);
+	if (bResult && InfoMessages.Num() > 0)
+	{
+		OutPlasticScmVersion = InfoMessages[0];
+
+		// Command-line format output changed with version 8.0.16.3000, plugin will crash unless we demand the old format.
+		// For more details, see https://www.plasticscm.com/download/releasenotes/oldernotes?release=8.0.16.3000
+		bIsNewVersion80163000 = !PlasticScmVersionLess(OutPlasticScmVersion, "8.0.16.3000");
+	}
+}
+
+void GetUserName(FString& OutUserName)
+{
+	TArray<FString> InfoMessages;
+	TArray<FString> ErrorMessages;
+	const bool bResult = RunCommand(TEXT("whoami"), TArray<FString>(), TArray<FString>(), EConcurrency::Synchronous, InfoMessages, ErrorMessages);
+	if (bResult && InfoMessages.Num() > 0)
+	{
+		OutUserName = InfoMessages[0];
+	}
+}
+
+bool GetWorkspaceName(FString& OutWorkspaceName)
+{
+	TArray<FString> InfoMessages;
+	TArray<FString> ErrorMessages;
+	TArray<FString> Parameters;
+	Parameters.Add(TEXT("."));
+	Parameters.Add(TEXT("--format={0}"));
+	// Get the workspace name
+	const bool bResult = RunCommand(TEXT("getworkspacefrompath"), Parameters, TArray<FString>(), EConcurrency::Synchronous, InfoMessages, ErrorMessages);
+	if (bResult && InfoMessages.Num() > 0)
+	{
+		// NOTE: getworkspacefrompath never returns an error!
+		if (!InfoMessages[0].Equals(TEXT(". is not in a workspace.")))
+		{
+			OutWorkspaceName = MoveTemp(InfoMessages[0]);
+		}
+	}
+
+	return bResult;
+}
+
+static bool ParseWorkspaceInformation(const TArray<FString>& InInfoMessages, int32& OutChangeset, FString& OutRepositoryName, FString& OutServerUrl, FString& OutBranchName)
+{
+	bool bResult = true;
+
+	// Get workspace status, in the form "cs:41@rep:UE4PlasticPlugin@repserver:localhost:8087" (disabled by the "--nostatus" flag)
+	//                                or "cs:41@rep:UE4PlasticPlugin@repserver:SRombauts@cloud" (when connected directly to the cloud)
+	if (InInfoMessages.Num() > 0)
+	{
+		static const FString ChangesetPrefix(TEXT("cs:"));
+		static const FString RepPrefix(TEXT("@rep:"));
+		static const FString ServerPrefix(TEXT("@repserver:"));
+		const FString& WorkspaceStatus = InInfoMessages[0];
+		const int32 RepIndex = WorkspaceStatus.Find(RepPrefix, ESearchCase::CaseSensitive);
+		const int32 ServerIndex = WorkspaceStatus.Find(ServerPrefix, ESearchCase::CaseSensitive, ESearchDir::FromEnd);
+		if ((RepIndex > INDEX_NONE) && (ServerIndex > INDEX_NONE))
+		{
+			const FString ChangesetString = WorkspaceStatus.Mid(ChangesetPrefix.Len(), RepIndex - ChangesetPrefix.Len());
+			OutChangeset = FCString::Atoi(*ChangesetString);
+			OutRepositoryName = WorkspaceStatus.Mid(RepIndex + RepPrefix.Len(), ServerIndex - RepIndex - RepPrefix.Len());
+			OutServerUrl = WorkspaceStatus.RightChop(ServerIndex + ServerPrefix.Len());
+		}
+		else
+		{
+			bResult = false;
+		}
+	}
+	// Get the branch name, in the form "Branch /main@UE4PlasticPluginDev" (enabled by the "--wkconfig" flag)
+	if (InInfoMessages.Num() > 1)
+	{
+		OutBranchName = InInfoMessages[1];
+	}
+
+	return bResult;
+}
+
+bool GetWorkspaceInformation(int32& OutChangeset, FString& OutRepositoryName, FString& OutServerUrl, FString& OutBranchName)
+{
+	TArray<FString> InfoMessages;
+	TArray<FString> ErrorMessages;
+	TArray<FString> Parameters;
+
+	// Command-line format output changed with version 8.0.16.3000, plugin will crash unless we demand the old format.
+	// For more details, see https://www.plasticscm.com/download/releasenotes/oldernotes?release=8.0.16.3000
+	if (bIsNewVersion80163000) {
+		Parameters.Add(TEXT("--compact"));
+	}
+	Parameters.Add(TEXT("--wkconfig")); // Branch name. TODO Deprecated in 8.0.16.3000 https://www.plasticscm.com/download/releasenotes/oldernotes?release=8.0.16.3000 => by default in xml format
+	Parameters.Add(TEXT("--nochanges")); // No file status // TODO: The option "--nochanges" has been renamed to "--header", and outputs just the workspace status.
+	bool bResult = RunCommand(TEXT("status"), Parameters, TArray<FString>(), EConcurrency::Synchronous, InfoMessages, ErrorMessages);
+	if (bResult)
+	{
+		bResult = ParseWorkspaceInformation(InfoMessages, OutChangeset, OutRepositoryName, OutServerUrl, OutBranchName);
+	}
+
+	return bResult;
+}
+
+/**
+ * @brief Extract the renamed from filename from a Plastic SCM status result.
+ *
+ * Examples of status results:
+ MV 100% Content\ToMove_BP.uasset -> Content\Moved_BP.uasset
+ *
+ * @param[in] InResult One line of status
+ * @return Renamed from filename extracted from the line of status
+ *
+ * @see FilenameFromPlasticStatus()
+ */
+static FString RenamedFromPlasticStatus(const FString& InResult)
+{
+	FString RenamedFrom;
+	int32 RenameIndex;
+	if (InResult.FindLastChar('>', RenameIndex))
+	{
+		// Extract only the first part of a rename "from -> to" (after the 2 letters status surrounded by 2 spaces)
+		RenamedFrom = InResult.Mid(9, RenameIndex - 9 - 2);
+	}
+	return RenamedFrom;
+}
+
+/**
+ * @brief Extract the relative filename from a Plastic SCM status result.
+ *
+ * Examples of status results:
+ CO Content\CheckedOut_BP.uasset
+ MV 100% Content\ToMove_BP.uasset -> Content\Moved_BP.uasset
+ *
+ * @param[in] InResult One line of status
+ * @return Relative filename extracted from the line of status
+ *
+ * @see FPlasticStatusFileMatcher and StateFromPlasticStatus()
+ */
+static FString FilenameFromPlasticStatus(const FString& InResult)
+{
+	FString RelativeFilename;
+	int32 RenameIndex;
+	if (InResult.FindLastChar('>', RenameIndex))
+	{
+		// Extract only the second part of a rename "from -> to"
+		RelativeFilename = InResult.RightChop(RenameIndex + 2);
+	}
+	else
+	{
+		// Extract the relative filename from the Plastic SCM status result (after the 2 letters status surrounded by 2 spaces)
+		RelativeFilename = InResult.RightChop(4);
+	}
+
+	return RelativeFilename;
+}
+
+/**
+ * @brief Match the relative filename of a Plastic SCM status result with a provided absolute filename
+ *
+ * Examples of status results:
+ CO Content\CheckedOut_BP.uasset
+ MV 100% Content\ToMove_BP.uasset -> Content\Moved_BP.uasset
+ */
+class FPlasticStatusFileMatcher
+{
+public:
+	FPlasticStatusFileMatcher(const FString& InAbsoluteFilename)
+		: AbsoluteFilename(InAbsoluteFilename)
+	{
+	}
+
+	bool operator()(const FString& InResult) const
+	{
+		return AbsoluteFilename.Contains(FilenameFromPlasticStatus(InResult));
+	}
+
+private:
+	const FString& AbsoluteFilename;
+};
+
+/**
+ * Extract and interpret the file state from the given Plastic "status" result.
+ * empty string = unmodified/controlled or hidden changes
+ CH Content\Changed_BP.uasset
+ CO Content\CheckedOut_BP.uasset
+ CP Content\Copied_BP.uasset
+ RP Content\Replaced_BP.uasset
+ AD Content\Added_BP.uasset
+ PR Content\Private_BP.uasset
+ IG Content\Ignored_BP.uasset
+ DE Content\Deleted_BP.uasset
+ LD Content\Deleted2_BP.uasset
+ MV 100% Content\ToMove_BP.uasset -> Content\Moved_BP.uasset
+ LM 100% Content\ToMove2_BP.uasset -> Content\Moved2_BP.uasset
+ */
+static EWorkspaceState::Type StateFromPlasticStatus(const FString& InResult)
+{
+	EWorkspaceState::Type State;
+	const FString FileStatus = InResult.Mid(1, 2);
+
+	if (FileStatus == "CH") // Modified but not Checked-Out
+	{
+		State = EWorkspaceState::Changed;
+	}
+	else if (FileStatus == "CO") // Checked-Out for modification
+	{
+		State = EWorkspaceState::CheckedOut;
+	}
+	else if (FileStatus == "CP")
+	{
+		State = EWorkspaceState::Copied;
+	}
+	else if (FileStatus == "RP")
+	{
+		State = EWorkspaceState::Replaced;
+	}
+	else if (FileStatus == "AD")
+	{
+		State = EWorkspaceState::Added;
+	}
+	else if ((FileStatus == "PR") || (FileStatus == "LM")) // Not Controlled/Not in Depot/Untracked (or Locally Moved/Renamed)
+	{
+		State = EWorkspaceState::Private;
+	}
+	else if (FileStatus == "IG")
+	{
+		State = EWorkspaceState::Ignored;
+	}
+	else if (FileStatus == "DE")
+	{
+		State = EWorkspaceState::Deleted; // Deleted (removed from source control)
+	}
+	else if (FileStatus == "LD")
+	{
+		State = EWorkspaceState::LocallyDeleted; // Locally Deleted (ie. missing)
+	}
+	else if (FileStatus == "MV")
+	{
+		State = EWorkspaceState::Moved; // Moved/Renamed
+	}
+	else
+	{
+		UE_LOG(LogSourceControl, Warning, TEXT("Unknown file status '%s' (in line '%s')"), *FileStatus, *InResult);
+		State = EWorkspaceState::Unknown;
+	}
+
+	return State;
+}
+
+/**
+ * @brief Parse the array of strings results of a 'cm status --noheaders --controlledchanged --changed --localdeleted --localmoved --private --ignored' command
+ *
+ * Called in case of a regular status command for one or multiple files (not for a whole directory). 
+ *
+ * @param[in]	InFiles		List of files in a directory (never empty).
+ * @param[in]	InResults	Lines of results from the "status" command
+ * @param[out]	OutStates	States of files for witch the status has been gathered
+ *
+ * Example cm status results:
+ CH Content\Changed_BP.uasset
+ CO Content\CheckedOut_BP.uasset
+ CP Content\Copied_BP.uasset
+ RP Content\Replaced_BP.uasset
+ AD Content\Added_BP.uasset
+ PR Content\Private_BP.uasset
+ IG Content\Ignored_BP.uasset
+ DE Content\Deleted_BP.uasset
+ LD Content\Deleted2_BP.uasset
+ MV 100% Content\ToMove_BP.uasset -> Content\Moved_BP.uasset
+ LM 100% Content\ToMove2_BP.uasset -> Content\Moved2_BP.uasset
+ */
+static void ParseFileStatusResult(const TArray<FString>& InFiles, const TArray<FString>& InResults, TArray<FPlasticSourceControlState>& OutStates, int32& OutChangeset, FString& OutBranchName)
+{
+	const FPlasticSourceControlModule& PlasticSourceControl = FModuleManager::GetModuleChecked<FPlasticSourceControlModule>("PlasticSourceControl");
+	const FString& WorkingDirectory = PlasticSourceControl.GetProvider().GetPathToWorkspaceRoot();
+
+	// Parse the first two lines with Changeset number and Branch name
+	FString RepositoryName, ServerUrl;
+	ParseWorkspaceInformation(InResults, OutChangeset, RepositoryName, ServerUrl, OutBranchName);
+
+	// Iterate on each file explicitly listed in the command
+	for (const FString& File : InFiles)
+	{
+		FPlasticSourceControlState FileState(File);
+
+		// Search the file in the list of status
+		// NOTE: in case of rename by editor, there are two results: checked-out AND renamed
+		// => we want to get the second one, witch is always the rename, so we search from the end
+		int32 IdxResult = InResults.FindLastByPredicate(FPlasticStatusFileMatcher(File));
+		if (IdxResult != INDEX_NONE)
+		{
+			// File found in status results; only the case for "changed" files
+			FileState.WorkspaceState = StateFromPlasticStatus(InResults[IdxResult]);
+
+			// Extract the original name of a Moved/Renamed file
+			if (EWorkspaceState::Moved == FileState.WorkspaceState)
+			{
+				FileState.MovedFrom = FPaths::ConvertRelativePathToFull(WorkingDirectory, RenamedFromPlasticStatus(InResults[IdxResult]));
+			}
+		}
+		else
+		{
+			// File not found in status
+			if (FPaths::FileExists(File))
+			{
+				// usually means the file is unchanged, or is on Hidden changes
+				FileState.WorkspaceState = EWorkspaceState::Controlled; // Unchanged
+			}
+			else
+			{
+				// but also the case for newly created content: there is no file on disk until the content is saved for the first time
+				FileState.WorkspaceState = EWorkspaceState::Private; // Not Controlled
+			}
+		}
+		FileState.TimeStamp.Now();
+
+		// @todo: temporary debug log (only for the first few files)
+		if (OutStates.Num() < 20) UE_LOG(LogSourceControl, Log, TEXT("%s = %d:%s"), *File, static_cast<uint32>(FileState.WorkspaceState), FileState.ToString());
+
+		OutStates.Add(MoveTemp(FileState));
+	}
+	// @todo: temporary debug log (if too many files)
+	if (OutStates.Num() > 20) UE_LOG(LogSourceControl, Log, TEXT("[...] %d more files"), OutStates.Num() - 20);
+}
+
+/**
+ * @brief Parse the array of strings results of a 'cm status --noheaders --controlledchanged --changed --localdeleted --localmoved --private --ignored' command
+ *
+ * Called in case of a "directory status" (no file listed in the command) ONLY to detect Removed/Deleted files !
+ *
+ * @param[in]	InResults	Lines of results from the "status" command
+ * @param[out]	OutStates	States of files for witch the status has been gathered
+ *
+ * @see #ParseFileStatusResult() above for an example of a cm status results
+*/
+static void ParseDirectoryStatusResult(const TArray<FString>& InResults, TArray<FPlasticSourceControlState>& OutStates)
+{
+	const FPlasticSourceControlModule& PlasticSourceControl = FModuleManager::GetModuleChecked<FPlasticSourceControlModule>("PlasticSourceControl");
+	const FString& PathToPlasticBinary = PlasticSourceControl.AccessSettings().GetBinaryPath();
+	const FString& WorkingDirectory = PlasticSourceControl.GetProvider().GetPathToWorkspaceRoot();
+
+	// Iterate on each line of result of the status command
+	// NOTE: in case of rename by editor, there are two results: checkouted AND renamed
+	// => we want to get the second one, witch is always the rename, so we just iterate and the second state will overwrite the first one
+	for (const FString& Result : InResults)
+	{
+		const FString RelativeFilename = FilenameFromPlasticStatus(Result);
+		const FString File = FPaths::ConvertRelativePathToFull(WorkingDirectory, RelativeFilename);
+		const EWorkspaceState::Type WorkspaceState = StateFromPlasticStatus(Result);
+		if ((EWorkspaceState::Deleted == WorkspaceState) || (EWorkspaceState::LocallyDeleted == WorkspaceState))
+		{
+			FPlasticSourceControlState FileState(File);
+			FileState.WorkspaceState = WorkspaceState;
+			FileState.TimeStamp.Now();
+
+			// @todo: temporary debug log
+			UE_LOG(LogSourceControl, Log, TEXT("%s = %d:%s"), *File, static_cast<uint32>(FileState.WorkspaceState), FileState.ToString());
+
+			OutStates.Add(MoveTemp(FileState));
+		}
+	}
+}
+
+/// Visitor to list all files in subdirectory
+class FFileVisitor : public IPlatformFile::FDirectoryVisitor
+{
+public:
+	virtual bool Visit(const TCHAR* FilenameOrDirectory, bool bIsDirectory) override
+	{
+		if (!bIsDirectory)
+		{
+			Files.Add(FilenameOrDirectory);
+		}
+		return true;
+	}
+
+	TArray<FString> Files;
+};
+
+/**
+ * @brief Run a "status" command for a directory to get workspace file states
+ *
+ *  It is either a command for a whole directory (ie. "Content/", in case of "Submit to Source Control"),
+ * or for one or more files all on a same directory (by design, since we group files by directory in RunUpdateStatus())
+ *
+ * @param[in]	InFiles				List of files in a directory, or the path to the directory itself (never empty).
+ * @param[out]	OutErrorMessages	Error messages from the "status" command
+ * @param[out]	OutStates			States of files for witch the status has been gathered (distinct than InFiles in case of a "directory status")
+ * @param[out]	OutChangeset		The current Changeset Number
+ * @param[out]	OutBranchName		Name of the current checked-out branch
+ */
+static bool RunStatus(const TArray<FString>& InFiles, const EConcurrency::Type InConcurrency, TArray<FString>& OutErrorMessages, TArray<FPlasticSourceControlState>& OutStates, int32& OutChangeset, FString& OutBranchName)
+{
+	ensure(0 < InFiles.Num());
+
+	TArray<FString> Parameters;
+
+	// Command-line format output changed with version 8.0.16.3000, plugin will crash unless we demand the old format.
+	// For more details, see https://www.plasticscm.com/download/releasenotes/oldernotes?release=8.0.16.3000
+	if (bIsNewVersion80163000) {
+		Parameters.Add(TEXT("--compact"));
+	}
+	Parameters.Add(TEXT("--wkconfig")); // Branch name. TODO Deprecated in 8.0.16.3000 https://www.plasticscm.com/download/releasenotes/oldernotes?release=8.0.16.3000 => by default in xml format
+	Parameters.Add(TEXT("--noheaders"));
+	Parameters.Add(TEXT("--controlledchanged --changed --localdeleted --localmoved --private --ignored"));
+	// "cm status" only operate on one path (file or folder) at a time, so use one folder path for multiple files in a directory
+	const FString Path = FPaths::GetPath(*InFiles[0]);
+	TArray<FString> OnePath;
+	// Only one file: optim very useful for the .uproject file at the root to avoid parsing the whole repository
+	// (does not work if file does not exist anymore)
+	if ((1 == InFiles.Num()) && (FPaths::FileExists(InFiles[0])))
+	{
+		OnePath.Add(InFiles[0]);
+	}
+	else
+	{
+		OnePath.Add(Path);
+	}
+	TArray<FString> Results;
+	TArray<FString> ErrorMessages;
+	const bool bResult = RunCommand(TEXT("status"), Parameters, OnePath, InConcurrency, Results, ErrorMessages);
+	// Normalize paths in the result (convert all '\' to '/')
+	for (int32 IdxResult = 0; IdxResult < Results.Num(); IdxResult++)
+	{
+		FPaths::NormalizeFilename(Results[IdxResult]);
+	}
+	OutErrorMessages.Append(ErrorMessages);
+	if (bResult)
+	{
+		if (1 == InFiles.Num() && FPaths::DirectoryExists(InFiles[0]))
+		{
+			// 1) Special case for "status" of a directory: requires a specific parse logic.
+			//   (this is triggered by the "Submit to Source Control" top menu button)
+			// Find recursively all files in the directory: this enable getting the list of "Controlled" (unchanged) assets
+			FFileVisitor FileVisitor;
+			FPlatformFileManager::Get().GetPlatformFile().IterateDirectoryRecursively(*Path, FileVisitor);
+			// @todo: temporary debug log
+			UE_LOG(LogSourceControl, Log, TEXT("RunStatus(%s): 1) special case for status of a directory containing %d file(s) (%s)"), *InFiles[0], FileVisitor.Files.Num(), *Path);
+			ParseFileStatusResult(FileVisitor.Files, Results, OutStates, OutChangeset, OutBranchName);
+			// The above cannot detect assets removed / locally deleted since there is no file left to enumerate (either by the Content Browser or by File Manager)
+			// => so we also parse the status results to explicitly look for Removed/Deleted assets
+			Results.RemoveAt(0, 2); // Before that, remove the first two line Changeset, and BranchName
+			ParseDirectoryStatusResult(Results, OutStates);
+		}
+		else
+		{
+			// 2) General case for one or more files in the same directory.
+			// @todo: temporary debug log
+			UE_LOG(LogSourceControl, Log, TEXT("RunStatus(%s...): 2) general case for %d file(s) in a directory (%s)"), *InFiles[0], InFiles.Num(), *Path);
+			ParseFileStatusResult(InFiles, Results, OutStates, OutChangeset, OutBranchName);
+		}
+	}
+
+	return bResult;
+}
+
+// Parse the fileinfo output format "{RevisionChangeset};{RevisionHeadChangeset};{RepSpec};{LockedBy};{LockedWhere}"
+// for example "40;41;repo@server:port;srombauts;UE4PlasticPluginDev"
+class FPlasticFileinfoParser
+{
+public:
+	FPlasticFileinfoParser(const FString& InResult)
+	{
+		TArray<FString> Fileinfos;
+		const int32 NbElmts = InResult.ParseIntoArray(Fileinfos, TEXT(";"));
+		if (NbElmts >= 2)
+		{
+			RevisionChangeset = FCString::Atoi(*Fileinfos[0]);
+			RevisionHeadChangeset = FCString::Atoi(*Fileinfos[1]);
+			if (NbElmts >= 3)
+			{
+				RepSpec = MoveTemp(Fileinfos[2]);
+				if (NbElmts >=4)
+				{
+					LockedBy = MoveTemp(Fileinfos[3]);
+					if (NbElmts >= 5)
+					{
+						LockedWhere = MoveTemp(Fileinfos[4]);
+					}
+				}
+			}
+		}
+	}
+
+	int32 RevisionChangeset;
+	int32 RevisionHeadChangeset;
+	FString RepSpec;
+	FString LockedBy;
+	FString LockedWhere;
+};
+
+/** Parse the array of strings result of a 'cm fileinfo --format="{RevisionChangeset};{RevisionHeadChangeset};{RepSpec};{LockedBy};{LockedWhere}"' command
+ *
+ * Example cm fileinfo results:
+16;16;;
+14;15;;
+17;17;srombauts;Workspace_2
+ */
+static void ParseFileinfoResults(const TArray<FString>& InResults, TArray<FPlasticSourceControlState>& InOutStates)
+{
+	const FPlasticSourceControlModule& PlasticSourceControl = FModuleManager::GetModuleChecked<FPlasticSourceControlModule>("PlasticSourceControl");
+	const FPlasticSourceControlProvider& Provider = PlasticSourceControl.GetProvider();
+
+	// Iterate on all files and all status of the result (assuming same number of line of results than number of file states)
+	for (int32 IdxResult = 0; IdxResult < InResults.Num(); IdxResult++)
+	{
+		const FString& Fileinfo = InResults[IdxResult];
+		FPlasticSourceControlState& FileState = InOutStates[IdxResult];
+		const FString& File = FileState.GetFilename();
+		FPlasticFileinfoParser FileinfoParser(Fileinfo);
+
+		FileState.LocalRevisionChangeset = FileinfoParser.RevisionChangeset;
+		FileState.DepotRevisionChangeset = FileinfoParser.RevisionHeadChangeset;
+		FileState.RepSpec = FileinfoParser.RepSpec;
+		FileState.LockedBy = MoveTemp(FileinfoParser.LockedBy);
+		FileState.LockedWhere = MoveTemp(FileinfoParser.LockedWhere);
+
+		// If a file is locked by someone but not checked-out locally (or moved/renamed) this means it is locked by someone else or elsewhere
+		if ((FileState.WorkspaceState != EWorkspaceState::CheckedOut) && (FileState.WorkspaceState != EWorkspaceState::Moved) && (0 < FileState.LockedBy.Len())) 
+		{
+			// @todo: temporary debug log
+			UE_LOG(LogSourceControl, Warning, TEXT("LockedByOther(%s) by '%s!=%s' (or %s!=%s)"), *File, *FileState.LockedBy, *Provider.GetUserName(), *FileState.LockedWhere, *Provider.GetWorkspaceName());
+			FileState.WorkspaceState = EWorkspaceState::LockedByOther;
+		}
+
+		// @todo: temporary debug log (only for the first few files)
+		if (IdxResult < 20) UE_LOG(LogSourceControl, Log, TEXT("%s: %d;%d %s by '%s' (%s)"), *File, FileState.LocalRevisionChangeset, FileState.DepotRevisionChangeset, *FileState.RepSpec, *FileState.LockedBy, *FileState.LockedWhere);
+	}
+	// @todo: temporary debug log (if too many files)
+	if (InResults.Num() > 20) UE_LOG(LogSourceControl, Log, TEXT("[...] %d more files"), InResults.Num() - 20);
+}
+
+/**
+ * @brief Run a Plastic "fileinfo" command to update status of given files.
+ *
+ * @param[in]		InForceFileinfo		Also force execute the fileinfo command required to do get RepSpec of xlinks when getting history (or for diffs)
+ * @param[in]		InConcurrency		Is the command running in the background, or blocking the main thread
+ * @param[out]		OutErrorMessages	Error messages from the "fileinfo" command
+ * @param[in,out]	InOutStates			List of file states in the directory, gathered by the "status" command, completed by results of the "fileinfo" command
+ */
+static bool RunFileinfo(const bool InForceFileinfo, const EConcurrency::Type InConcurrency, TArray<FString>& OutErrorMessages, TArray<FPlasticSourceControlState>& InOutStates)
+{
+	bool bResult = true;
+	TArray<FString> Files;
+	for (const auto& State : InOutStates)
+	{
+		// Optimize by not issuing "fileinfo" commands on "Added"/"Deleted"/"NotControled"/"Ignored" but also "CheckedOut" and "Moved" files.
+		// This can greatly reduce the time needed to do some basic operation like "Add to source control" when using a distant server or the Plastic Cloud.
+		// this can't work with xlink file when we want to update the history
+		// we need to know that we are running a fileinfo command to get the history, that's the role of InForceFileinfo
+		if (	(InForceFileinfo)
+			||	(State.WorkspaceState == EWorkspaceState::Controlled)
+			||	(State.WorkspaceState == EWorkspaceState::Changed)
+			||	(State.WorkspaceState == EWorkspaceState::Replaced)
+			||	(State.WorkspaceState == EWorkspaceState::Conflicted)
+		//	||	(State.WorkspaceState == EWorkspaceState::LockedByOther) // we do not have this info at this stage, cf. ParseFileinfoResults()
+			)
+		{
+			Files.Add(State.GetFilename());
+		}
+	}
+	if (Files.Num() > 0)
+	{
+		TArray<FString> Results;
+		TArray<FString> ErrorMessages;
+		TArray<FString> Parameters;
+		Parameters.Add(TEXT("--format=\"{RevisionChangeset};{RevisionHeadChangeset};{RepSpec};{LockedBy};{LockedWhere}\""));
+		bResult = RunCommand(TEXT("fileinfo"), Parameters, Files, InConcurrency, Results, ErrorMessages);
+		OutErrorMessages.Append(ErrorMessages);
+		if (bResult)
+		{
+			ParseFileinfoResults(Results, InOutStates);
+		}
+	}
+
+	return bResult;
+}
+
+// FILE_CONFLICT /Content/FirstPersonBP/Blueprints/FirstPersonProjectile.uasset 1 4 6 903
+// (explanations: 'The file /Content/FirstPersonBP/Blueprints/FirstPersonProjectile.uasset needs to be merged from cs:4 to cs:6 base cs:1. Changed by both contributors.')
+class FPlasticMergeConflictParser
+{
+public:
+	FPlasticMergeConflictParser(const FString& InResult)
+	{
+		static const FString FILE_CONFLICT(TEXT("FILE_CONFLICT "));
+		if (InResult.StartsWith(FILE_CONFLICT, ESearchCase::CaseSensitive))
+		{
+			FString Temp = InResult.RightChop(FILE_CONFLICT.Len());
+			int32 WhitespaceIndex;
+			if (Temp.FindChar(TEXT(' '), WhitespaceIndex))
+			{
+				Filename = Temp.Left(WhitespaceIndex);
+			}
+			Temp = Temp.RightChop(WhitespaceIndex + 1);
+			if (Temp.FindChar(TEXT(' '), WhitespaceIndex))
+			{
+				const FString Base = Temp.Left(WhitespaceIndex);
+				BaseChangeset = FCString::Atoi(*Base);
+			}
+			Temp = Temp.RightChop(WhitespaceIndex + 1);
+			if (Temp.FindChar(TEXT(' '), WhitespaceIndex))
+			{
+				const FString Source = Temp.Left(WhitespaceIndex);
+				SourceChangeset = FCString::Atoi(*Source);
+			}
+		}
+	}
+
+	FString Filename;
+	int32 BaseChangeset;
+	int32 SourceChangeset;
+};
+
+// Check if merging, and from which changelist, then execute a cm merge command to amend status for listed files
+bool RunCheckMergeStatus(const TArray<FString>& InFiles, TArray<FString>& OutErrorMessages, TArray<FPlasticSourceControlState>& OutStates)
+{
+	bool bResult = false;
+	FPlasticSourceControlModule& PlasticSourceControl = FModuleManager::GetModuleChecked<FPlasticSourceControlModule>("PlasticSourceControl");
+	FPlasticSourceControlProvider& Provider = PlasticSourceControl.GetProvider();
+
+	const FString MergeProgressFilename = FPaths::Combine(*Provider.GetPathToWorkspaceRoot(), TEXT(".plastic/plastic.mergeprogress"));
+	if (FPaths::FileExists(MergeProgressFilename))
+	{
+		// read in file as string
+		FString MergeProgressContent;
+		if (FFileHelper::LoadFileToString(MergeProgressContent, *MergeProgressFilename))
+		{
+			// @todo: temporary debug logs
+			UE_LOG(LogSourceControl, Log, TEXT("RunCheckMergeStatus: %s:\n%s"), *MergeProgressFilename, *MergeProgressContent);
+			// Content is in one line, looking like the following:
+			// Target: mount:56e62dd7-241f-41e9-8c6b-dd4ca4513e62#/#UE4MergeTest@localhost:8087 merged from: Merge 4
+			// Target: mount:56e62dd7-241f-41e9-8c6b-dd4ca4513e62#/#UE4MergeTest@localhost:8087 merged from: Cherrypicking 3
+			// Target: mount:56e62dd7-241f-41e9-8c6b-dd4ca4513e62#/#UE4MergeTest@localhost:8087 merged from: IntervalCherrypick 2 4
+			// 1) Extract the word after "merged from: "
+			static const FString MergeFromString(TEXT("merged from: "));
+			const int32 MergeFromIndex = MergeProgressContent.Find(MergeFromString, ESearchCase::CaseSensitive);
+			if (MergeFromIndex > INDEX_NONE)
+			{
+				const FString MergeType = MergeProgressContent.RightChop(MergeFromIndex + MergeFromString.Len());
+				int32 SpaceBeforeChangesetIndex;
+				if (MergeType.FindChar(TEXT(' '), SpaceBeforeChangesetIndex))
+				{
+					// 2) In case of "Merge" or "Cherrypicking" extract the merge changelist xxx after the last space (use case for merge from "branch", from "label", and for "merge on Update")
+					const FString ChangesetString = MergeType.RightChop(SpaceBeforeChangesetIndex + 1);
+					const int32 Changeset = FCString::Atoi(*ChangesetString);
+					const FString ChangesetSpecification = FString::Printf(TEXT("cs:%d"), Changeset);
+
+					TArray<FString> Results;
+					TArray<FString> ErrorMessages;
+					TArray<FString> Parameters;
+					Parameters.Add(ChangesetSpecification);
+
+					int32 SpaceBeforeChangeset2Index;
+					if (ChangesetString.FindLastChar(TEXT(' '), SpaceBeforeChangeset2Index))
+					{
+						// 3) In case of "IntervalCherrypick", extract the 2 changelists
+						const FString Changeset2String = ChangesetString.RightChop(SpaceBeforeChangeset2Index + 1);
+						const int32 Changeset2 = FCString::Atoi(*Changeset2String);
+						const FString Changeset2Specification = FString::Printf(TEXT("--interval-origin=cs:%d"), Changeset2);
+
+						Parameters.Add(Changeset2Specification);
+					}
+					else
+					{
+						if (MergeType.StartsWith(TEXT("Cherrypicking"), ESearchCase::CaseSensitive))
+						{
+							Parameters.Add(TEXT("--cherrypicking"));
+						}
+					}
+					// Store the Merge Parameters for reuse with later "Resolve" operation
+					const TArray<FString> PendingMergeParameters = Parameters;
+					Parameters.Add(TEXT("--machinereadable"));
+					// call 'cm merge cs:xxx --machinereadable' (only dry-run, whithout the --merge parameter)
+					bResult = RunCommand(TEXT("merge"), Parameters, TArray<FString>(), EConcurrency::Synchronous, Results, ErrorMessages);
+					OutErrorMessages.Append(ErrorMessages);
+					// Parse the result, one line for each conflicted files:
+					for (const FString& Result : Results)
+					{
+						FPlasticMergeConflictParser MergeConflict(Result);
+						UE_LOG(LogSourceControl, Log, TEXT("MergeConflict.Filename: '%s'"), *MergeConflict.Filename);
+						for (FPlasticSourceControlState& State : OutStates)
+						{
+							UE_LOG(LogSourceControl, Log, TEXT("State.LocalFilename: '%s'"), *State.LocalFilename);
+							if (State.LocalFilename.EndsWith(MergeConflict.Filename, ESearchCase::CaseSensitive))
+							{
+								// @todo: temporary debug log
+								UE_LOG(LogSourceControl, Log, TEXT("MergeConflict '%s' found Base cs:%d From cs:%d"), *MergeConflict.Filename, MergeConflict.BaseChangeset, MergeConflict.SourceChangeset);
+								State.WorkspaceState = EWorkspaceState::Conflicted;
+								State.PendingMergeFilename = MergeConflict.Filename;
+								State.PendingMergeBaseChangeset = MergeConflict.BaseChangeset;
+								State.PendingMergeSourceChangeset = MergeConflict.SourceChangeset;
+								State.PendingMergeParameters = PendingMergeParameters;
+								break;
+							}
+						}
+					}
+				}
+			}
+		}
+	}
+
+	return bResult;
+}
+
+// Run a batch of Plastic "status" and "fileinfo" commands to update status of given files and directories.
+bool RunUpdateStatus(const TArray<FString>& InFiles, const bool InForceFileinfo, const EConcurrency::Type InConcurrency, TArray<FString>& OutErrorMessages, TArray<FPlasticSourceControlState>& OutStates, int32& OutChangeset, FString& OutBranchName)
+{
+	bool bResults = true;
+
+	// The "status" command only operate on one directory at a time
+	// (whole tree recursively) not on different folders with no common root.
+	// But "Submit to Source Control" ask for the State of 3 different directories, Engine/Content, Project/Content and Project/Config,
+	// In a same way, a check-in can involve files from different subdirectories, and UpdateStatus is called for them all at once.
+
+	// 1) So here we group files by path (ie. by subdirectory)
+	TMap<FString, TArray<FString>> GroupOfFiles;
+	for (const FString& File : InFiles)
+	{
+		const FString Path = FPaths::GetPath(*File);
+		TArray<FString>* Group = GroupOfFiles.Find(Path);
+		if (Group != nullptr)
+		{
+			Group->Add(File);
+		}
+		else
+		{
+			TArray<FString> NewGroup;
+			NewGroup.Add(File);
+			GroupOfFiles.Add(Path, NewGroup);
+		}
+	}
+
+	// @todo: temporary debug log
+	if (InFiles.Num() > 0)
+	{
+		UE_LOG(LogSourceControl, Log, TEXT("RunUpdateStatus: %d file(s)/%d directory(ies) ('%s'...)"), InFiles.Num(), GroupOfFiles.Num(), *InFiles[0]);
+	}
+	else
+	{
+		UE_LOG(LogSourceControl, Warning, TEXT("RunUpdateStatus: NO file"));
+	}
+
+	// 2) then we can batch Plastic status operation by subdirectory
+	for (const auto& Files : GroupOfFiles)
+	{
+		// Run a "status" command on the directory to get workspace file states.
+		TArray<FPlasticSourceControlState> States;
+		const bool bGroupOk = RunStatus(Files.Value, InConcurrency, OutErrorMessages, States, OutChangeset, OutBranchName);
+		if (bGroupOk && (States.Num() > 0))
+		{
+			// Run a "fileinfo" command to update status of given files.
+			// In case of "directory status", there is no explicit file in the group (it contains only the directory) 
+			// => work on the list of files discovered by RunStatus()
+			bResults &= RunFileinfo(InForceFileinfo, InConcurrency, OutErrorMessages, States);
+		}
+		OutStates.Append(MoveTemp(States));
+	}
+
+	// Check if merging, and from which changelist, then execute a cm merge command to amend status for listed files
+	RunCheckMergeStatus(InFiles, OutErrorMessages, OutStates);
+
+	return bResults;
+}
+
+// Run a Plastic "cat" command to dump the binary content of a revision into a file.
+// cm cat revid:1230@rep:myrep@repserver:myserver:8084 --raw --file=Name124.tmp
+bool RunDumpToFile(const FString& InPathToPlasticBinary, const FString& InRevSpec, const FString& InDumpFileName)
+{
+	int32	ReturnCode = 0;
+	FString Results;
+	FString Errors;
+
+	// start with the Plastic command itself, then add revspec and temp filename to dump
+	FString FullCommand = TEXT("cat ");
+	FullCommand += InRevSpec;
+	FullCommand += TEXT(" --raw --file=\"");
+	FullCommand += InDumpFileName;
+	FullCommand += TEXT("\"");
+
+	// @todo: temporary debug logs
+	UE_LOG(LogSourceControl, Log, TEXT("RunDumpToFile: '%s %s'"), *InPathToPlasticBinary, *FullCommand);
+	const bool bResult = FPlatformProcess::ExecProcess(*InPathToPlasticBinary, *FullCommand, &ReturnCode, &Results, &Errors);
+	UE_LOG(LogSourceControl, Log, TEXT("RunDumpToFile: ExecProcess ReturnCode=%d Results='%s'"), ReturnCode, *Results);
+	if (!bResult || !Errors.IsEmpty())
+	{
+		UE_LOG(LogSourceControl, Error, TEXT("RunDumpToFile: ExecProcess ReturnCode=%d Errors='%s'"), ReturnCode, *Errors);
+	}
+
+	return bResult;
+}
+
+
+/**
+ * Translate file actions from Plastic 'cm log' command to keywords used by the Editor UI.
+ *
+ * @see SHistoryRevisionListRowContent::GenerateWidgetForColumn(): "add", "edit", "delete", "branch" and "integrate"
+*/
+FString TranslateAction(const FString& InAction)
+{
+	if (InAction.Equals(TEXT("Added")))
+	{
+		return TEXT("add");
+	}
+	else if (InAction.Equals(TEXT("Moved")))
+	{
+		return TEXT("branch");
+	}
+	else if (InAction.Equals(TEXT("Deleted")))
+	{
+		return TEXT("delete");
+	}
+	else // if (InAction.Equals(TEXT("Changed")))
+	{
+		return TEXT("edit");
+	}
+}
+
+/**
+ * Parse the array of strings results of a 'cm log --xml' command
+ *
+ * Example cm log results:
+<?xml version="1.0" encoding="utf-8"?>
+<LogList>
+  <Changeset>
+    <ObjId>989</ObjId>
+    <ChangesetId>2</ChangesetId>
+    <Branch>/main</Branch>
+    <Comment>Ignore Collections and Developers content</Comment>
+    <Owner>dev</Owner>
+    <GUID>a985c487-0f54-45c5-b0ef-9b87c4c3c3f9</GUID>
+    <Changes>
+      <Item>
+        <Branch>/main</Branch>
+        <RevNo>2</RevNo>
+        <Owner>dev</Owner>
+        <RevId>985</RevId>
+        <ParentRevId>282</ParentRevId>
+        <SrcCmPath>/ignore.conf</SrcCmPath>
+        <SrcParentItemId>2</SrcParentItemId>
+        <DstCmPath>/ignore.conf</DstCmPath>
+        <DstParentItemId>2</DstParentItemId>
+        <Date>2016-04-18T10:44:49.0000000+02:00</Date>
+        <Type>Changed</Type>
+      </Item>
+    </Changes>
+    <Date>2016-04-18T10:44:49.0000000+02:00</Date>
+  </Changeset>
+</LogList>
+*/
+static void ParseLogResults(const FXmlFile& InXmlResult, FPlasticSourceControlRevision& OutSourceControlRevision)
+{
+	static const FString LogList(TEXT("LogList"));
+	static const FString Changeset(TEXT("Changeset"));
+	static const FString Comment(TEXT("Comment"));
+	static const FString Date(TEXT("Date"));
+	static const FString Owner(TEXT("Owner"));
+	static const FString Changes(TEXT("Changes"));
+	static const FString Item(TEXT("Item"));
+	static const FString RevId(TEXT("RevId"));
+	static const FString ParentRevId(TEXT("ParentRevId"));
+	static const FString SrcCmPath(TEXT("SrcCmPath"));
+	static const FString DstCmPath(TEXT("DstCmPath"));
+	static const FString Type(TEXT("Type"));
+
+	const FXmlNode* LogListNode = InXmlResult.GetRootNode();
+	if (LogListNode == nullptr || LogListNode->GetTag() != LogList)
+	{
+		return;
+	}
+
+	const FXmlNode* ChangesetNode = LogListNode->FindChildNode(Changeset);
+	if (ChangesetNode == nullptr)
+	{
+		return;
+	}
+
+	const FXmlNode* CommentNode = ChangesetNode->FindChildNode(Comment);
+	if (CommentNode != nullptr)
+	{
+		OutSourceControlRevision.Description = CommentNode->GetContent();
+	}
+	const FXmlNode* OwnerNode = ChangesetNode->FindChildNode(Owner);
+	if (OwnerNode != nullptr)
+	{
+		OutSourceControlRevision.UserName = OwnerNode->GetContent();
+	}
+	const FXmlNode* DateNode = ChangesetNode->FindChildNode(Date);
+	if (DateNode != nullptr)
+	{	//                           |--|
+		//    2016-04-18T10:44:49.0000000+02:00
+		// => 2016-04-18T10:44:49.000+02:00
+		const FString DateIso = DateNode->GetContent().LeftChop(10) + DateNode->GetContent().RightChop(27);
+		FDateTime::ParseIso8601(*DateIso, OutSourceControlRevision.Date);
+	}
+
+	const FXmlNode* ChangesNode = ChangesetNode->FindChildNode(Changes);
+	if (ChangesNode == nullptr)
+	{
+		return;
+	}
+
+	// Iterate on files to find the one we are tracking
+	for (const FXmlNode* ItemNode : ChangesNode->GetChildrenNodes())
+	{
+		int32 RevisionId = -1;
+		const FXmlNode* RevIdNode = ItemNode->FindChildNode(RevId);
+		if (RevIdNode != nullptr)
+		{
+			RevisionId = FCString::Atoi(*RevIdNode->GetContent());
+		}
+		// Is this about the file we are looking for?
+		if (RevisionId == OutSourceControlRevision.RevisionId)
+		{
+			const FXmlNode* DstCmPathNode = ItemNode->FindChildNode(DstCmPath);
+			if (DstCmPathNode != nullptr)
+			{
+				OutSourceControlRevision.Filename = DstCmPathNode->GetContent();
+
+				const FXmlNode* SrcCmPathNode = ItemNode->FindChildNode(SrcCmPath);
+				const FXmlNode* ParentRevIdNode = ItemNode->FindChildNode(ParentRevId);
+				// Detect case of rename ("branch" in Perforce vocabulary)
+				if (ParentRevIdNode != nullptr && SrcCmPathNode != nullptr && !SrcCmPathNode->GetContent().Equals(DstCmPathNode->GetContent()))
+				{
+					TSharedRef<FPlasticSourceControlRevision, ESPMode::ThreadSafe> MovedFromRevision = MakeShareable(new FPlasticSourceControlRevision);
+					MovedFromRevision->Filename = SrcCmPathNode->GetContent();
+					MovedFromRevision->RevisionId = FCString::Atoi(*ParentRevIdNode->GetContent());
+	
+					OutSourceControlRevision.BranchSource = MovedFromRevision;
+				}
+			}
+			const FXmlNode* TypeNode = ItemNode->FindChildNode(Type);
+			if (TypeNode != nullptr)
+			{
+				OutSourceControlRevision.Action = TranslateAction(TypeNode->GetContent());
+			}
+			// 	Do not stop at first match, because in case of rename there are multiple log nodes: Changed+Moved (in this order)
+		}
+	}
+}
+
+// Run "cm log" on the changeset provided by the "history" command to get extra info about the change at a specific revision
+static bool RunLogCommand(const FString& InChangeset, const FString& InRepSpec, FPlasticSourceControlRevision& OutSourceControlRevision)
+{
+	const FString RepositorySpecification = FString::Printf(TEXT("cs:%s@%s"), *InChangeset, *InRepSpec);
+
+	FString Results;
+	FString Errors;
+	TArray<FString> Parameters;
+	Parameters.Add(RepositorySpecification);
+	Parameters.Add(TEXT("--xml"));
+	Parameters.Add(TEXT("--encoding=\"utf-8\""));
+
+	// Uses the raw RunCommandInternal() that does not split results in an array of strings, for XML parsing
+	bool bResult = RunCommandInternal(TEXT("log"), Parameters, TArray<FString>(), EConcurrency::Synchronous, Results, Errors);
+	if (bResult)
+	{
+		FXmlFile XmlFile;
+		bResult = XmlFile.LoadFile(Results, EConstructMethod::ConstructFromBuffer);
+		if (bResult)
+		{
+			ParseLogResults(XmlFile, OutSourceControlRevision);
+		}
+	}
+	return bResult;
+}
+
+/**
+ * Parse results of the 'cm history --format="{1};{6}"' command ("Changeset number" and "Revision id"),
+ * then run "cm log" on each revision to get extra info about the change (description, date, filename, branch, action)
+ * 
+ * Results of the history command are with one changeset number and revision id by line, like that:
+14;176
+17;220
+18;223
+*/
+static bool ParseHistoryResults(const TArray<FString>& InResults, FPlasticSourceControlState& InOutState)
+{
+	bool bResult = true;
+
+	const FPlasticSourceControlModule& PlasticSourceControl = FModuleManager::GetModuleChecked<FPlasticSourceControlModule>("PlasticSourceControl");
+	const FPlasticSourceControlProvider& Provider = PlasticSourceControl.GetProvider();
+	const FString RootRepSpec = FString::Printf(TEXT("%s@%s"), *Provider.GetRepositoryName(), *Provider.GetServerUrl());
+
+	InOutState.History.Reserve(InResults.Num());
+
+	// parse history in reverse: needed to get most recent at the top (implied by the UI)
+	for (int32 Index = InResults.Num() - 1; Index >= 0; Index--)
+	{
+		const FString& Result = InResults[Index];
+		if (bResult)
+		{
+			TArray<FString> Infos;
+			const int32 NbElmts = Result.ParseIntoArray(Infos, TEXT(";"));
+			if (NbElmts == 2)
+			{
+				const TSharedRef<FPlasticSourceControlRevision, ESPMode::ThreadSafe> SourceControlRevision = MakeShareable(new FPlasticSourceControlRevision);
+				SourceControlRevision->State = &InOutState;
+				const FString& Changeset = Infos[0];
+				const FString& RevisionId = Infos[1];
+				SourceControlRevision->ChangesetNumber = FCString::Atoi(*Changeset); // Value now used in the Revision column and in the Asset Menu History
+				SourceControlRevision->RevisionId = FCString::Atoi(*RevisionId); // 
+				// Also append depot name to the revision, but only when it is different from the default one (ie for xlinks sub repository)
+				if (InOutState.RepSpec != RootRepSpec)
+				{
+					TArray<FString> RepSpecs;
+					InOutState.RepSpec.ParseIntoArray(RepSpecs, TEXT("@"));
+					SourceControlRevision->Revision = FString::Printf(TEXT("cs:%s@%s"), *Changeset, *RepSpecs[0]);
+				}
+				else
+				{
+					SourceControlRevision->Revision = FString::Printf(TEXT("cs:%s"), *Changeset);
+				}
+
+				// Run "cm log" on the changeset number
+				bResult = RunLogCommand(Changeset, InOutState.RepSpec, *SourceControlRevision);
+				InOutState.History.Add(SourceControlRevision);
+			}
+			else
+			{
+				bResult = false;
+			}
+		}
+	}
+
+	return bResult;
+}
+
+// Run a Plastic "history" command and multiple "log" commands and parse them.
+bool RunGetHistory(const FString& InFile, TArray<FString>& OutErrorMessages, FPlasticSourceControlState& InOutState)
+{
+	TArray<FString> Results;
+	TArray<FString> Parameters;
+	Parameters.Add(TEXT("--format=\"{1};{6}\"")); // Get "Changeset number" and "Revision id" of each revision of the asset
+	TArray<FString> OneFile;
+	OneFile.Add(*InFile);
+
+	bool bResult = RunCommand(TEXT("history"), Parameters, OneFile, EConcurrency::Synchronous, Results, OutErrorMessages);
+	if (bResult)
+	{
+		bResult = ParseHistoryResults(Results, InOutState);
+	}
+
+	return bResult;
+}
+
+bool UpdateCachedStates(TArray<FPlasticSourceControlState>&& InStates)
+{
+	FPlasticSourceControlModule& PlasticSourceControl = FModuleManager::GetModuleChecked<FPlasticSourceControlModule>( "PlasticSourceControl" );
+	FPlasticSourceControlProvider& Provider = PlasticSourceControl.GetProvider();
+	const FDateTime Now = FDateTime::Now();
+
+	for (auto&& InState : InStates)
+	{
+		TSharedRef<FPlasticSourceControlState, ESPMode::ThreadSafe> State = Provider.GetStateInternal(InState.LocalFilename);
+		*State = MoveTemp(InState);
+		State->TimeStamp = Now;
+	}
+
+	return (InStates.Num() > 0);
+}
+
+/**
+ * Helper struct for RemoveRedundantErrors()
+ */
+struct FRemoveRedundantErrors
+{
+	FRemoveRedundantErrors(const FString& InFilter)
+		: Filter(InFilter)
+	{
+	}
+
+	bool operator()(const FString& String) const
+	{
+		if(String.Contains(Filter))
+		{
+			return true;
+		}
+
+		return false;
+	}
+
+	/** The filter string we try to identify in the reported error */
+	FString Filter;
+};
+
+void RemoveRedundantErrors(FPlasticSourceControlCommand& InCommand, const FString& InFilter)
+{
+	bool bFoundRedundantError = false;
+	for(const FString& ErrorMessage : InCommand.ErrorMessages)
+	{
+		if(ErrorMessage.Contains(InFilter, ESearchCase::CaseSensitive, ESearchDir::FromEnd))
+		{
+			InCommand.InfoMessages.Add(ErrorMessage);
+			bFoundRedundantError = true;
+		}
+	}
+
+	InCommand.ErrorMessages.RemoveAll( FRemoveRedundantErrors(InFilter) );
+
+	// if we have no error messages now, assume success!
+	if(bFoundRedundantError && InCommand.ErrorMessages.Num() == 0 && !InCommand.bCommandSuccessful)
+	{
+		InCommand.bCommandSuccessful = true;
+	}
+}
+
+}
diff --git a/Engine/Plugins/Developer/PlasticSourceControl/Source/PlasticSourceControl/Private/PlasticSourceControlUtils.h b/Engine/Plugins/Developer/PlasticSourceControl/Source/PlasticSourceControl/Private/PlasticSourceControlUtils.h
new file mode 100644
--- /dev/null
+++ b/Engine/Plugins/Developer/PlasticSourceControl/Source/PlasticSourceControl/Private/PlasticSourceControlUtils.h
@@ -0,0 +1,145 @@
+// Copyright 1998-2019 Epic Games, Inc. All Rights Reserved.
+
+#pragma once
+
+#include "CoreMinimal.h"
+#include "ISourceControlProvider.h"
+#include "PlasticSourceControlRevision.h"
+
+class FPlasticSourceControlCommand;
+class FPlasticSourceControlState;
+
+/**
+ * Helper struct for maintaining temporary files for passing to commands
+ */
+class FScopedTempFile
+{
+public:
+
+	/** Constructor - open & write string to temp file */
+	FScopedTempFile(const FText& InText);
+
+	/** Destructor - delete temp file */
+	~FScopedTempFile();
+
+	/** Get the filename of this temp file - empty if it failed to be created */
+	const FString& GetFilename() const;
+
+private:
+	/** The filename we are writing to */
+	FString Filename;
+};
+
+namespace PlasticSourceControlUtils
+{
+
+/**
+ * Find the path to the Plastic binary: for now relying on the Path to access the "cm" command.
+ */
+FString FindPlasticBinaryPath();
+
+/**
+ * Launch the Plastic SCM "shell" command to run it in background.
+ * @param	InPathToPlasticBinary	The path to the Plastic binary
+ * @param	InWorkspaceRoot			The workspace from where to run the command - usually the Game directory
+ * @returns true if the command succeeded and returned no errors
+ */
+bool LaunchBackgroundPlasticShell(const FString& InPathToPlasticBinary, const FString& InWorkspaceRoot);
+
+/** Terminate the background 'cm shell' process and associated pipes */
+void Terminate();
+
+/**
+ * Find the root of the Plastic workspace, looking from the GameDir and upward in its parent directories
+ * @param InPathToGameDir		The path to the Game Directory
+ * @param OutWorkspaceRoot		The path to the root directory of the Plastic workspace if found, else the path to the GameDir
+ * @returns true if the command succeeded and returned no errors
+ */
+bool FindRootDirectory(const FString& InPathToGameDir, FString& OutWorkspaceRoot);
+
+/**
+ * Get Plastic SCM cli version
+ * @param	OutCliVersion		Version of the Plastic SCM Command Line Interface tool
+*/
+void GetPlasticScmVersion(FString& OutPlasticScmVersion);
+
+/**
+ * Get Plastic SCM current user
+ * @param	OutUserName			Name of the Plastic SCM user configured globally
+ */
+void GetUserName(FString& OutUserName);
+
+/**
+ * Get Plastic workspace name
+ * @param	OutWorkspaceName	Name of the current workspace
+ */
+bool GetWorkspaceName(FString& OutWorkspaceName);
+
+/**
+ * Get Plastic repository name and server URL, branch name and current changeset number
+ * @param	OutChangeset		The current Changeset Number
+ * @param	OutRepositoryName	Name of the repository of the current workspace
+ * @param	OutServerUrl		Url/Port of the server of the repository
+ * @param	OutBranchName		Name of the current checked-out branch
+ */
+bool GetWorkspaceInformation(int32& OutChangeset, FString& OutRepositoryName, FString& OutServerUrl, FString& OutBranchName);
+
+/**
+ * Run a Plastic command - output is a string TArray.
+ *
+ * @param	InCommand			The Plastic command - e.g. commit
+ * @param	InParameters		The parameters to the Plastic command
+ * @param	InFiles				The files to be operated on
+ * @param	InConcurrency		Is the command running in the background, or bloking the main thread
+ * @param	OutResults			The results (from StdOut) as an array per-line
+ * @param	OutErrorMessages	Any errors (from StdErr) as an array per-line
+ * @returns true if the command succeeded and returned no errors
+ */
+bool RunCommand(const FString& InCommand, const TArray<FString>& InParameters, const TArray<FString>& InFiles, const EConcurrency::Type InConcurrency, TArray<FString>& OutResults, TArray<FString>& OutErrorMessages);
+
+/**
+ * Run a Plastic "status" command and parse it.
+ *
+ * @param	InFiles				The files to be operated on
+ * @param	InForceFileinfo		Also force execute the fileinfo command required to do get RepSpec of xlinks when getting history (or for diffs)
+ * @param	InConcurrency		Is the command running in the background, or blocking the main thread
+ * @param	OutErrorMessages	Any errors (from StdErr) as an array per-line
+ * @param	OutStates			States of the files
+ * @param	OutChangeset		The current Changeset Number
+ * @param	OutBranchName		Name of the current checked-out branch
+ * @returns true if the command succeeded and returned no errors
+ */
+bool RunUpdateStatus(const TArray<FString>& InFiles, const bool InForceFileinfo, const EConcurrency::Type InConcurrency, TArray<FString>& OutErrorMessages, TArray<FPlasticSourceControlState>& OutStates, int32& OutChangeset, FString& OutBranchName);
+
+/**
+ * Run a Plastic "cat" command to dump the binary content of a revision into a file.
+ *
+ * @param	InPathToPlasticBinary	The path to the Plastic binary
+ * @param	InRevSpec				The revision specification to get
+ * @param	InDumpFileName			The temporary file to dump the revision
+ * @returns true if the command succeeded and returned no errors
+*/
+bool RunDumpToFile(const FString& InPathToPlasticBinary, const FString& InRevSpec, const FString& InDumpFileName);
+
+/**
+ * Run a Plastic "history" and "log" commands and parse it.
+ *
+ * @param	InFile				The file to be operated on
+ * @param	OutErrorMessages	Any errors (from StdErr) as an array per-line
+ * @param	InOutState			The status to update with the history of the file
+ */
+bool RunGetHistory(const FString& InFile, TArray<FString>& OutErrorMessages, FPlasticSourceControlState& InOutState);
+
+/**
+ * Helper function for various commands to update cached states.
+ * @returns true if any states were updated
+ */
+bool UpdateCachedStates(TArray<FPlasticSourceControlState>&& InStates);
+
+/** 
+ * Remove redundant errors (that contain a particular string) and also
+ * update the commands success status if all errors were removed.
+ */
+void RemoveRedundantErrors(FPlasticSourceControlCommand& InCommand, const FString& InFilter);
+
+}
diff --git a/Engine/Plugins/Developer/PlasticSourceControl/Source/PlasticSourceControl/Private/SPlasticSourceControlSettings.cpp b/Engine/Plugins/Developer/PlasticSourceControl/Source/PlasticSourceControl/Private/SPlasticSourceControlSettings.cpp
new file mode 100644
--- /dev/null
+++ b/Engine/Plugins/Developer/PlasticSourceControl/Source/PlasticSourceControl/Private/SPlasticSourceControlSettings.cpp
@@ -0,0 +1,680 @@
+// Copyright 1998-2019 Epic Games, Inc. All Rights Reserved.
+
+#include "SPlasticSourceControlSettings.h"
+#include "PlasticSourceControlOperations.h"
+#include "PlasticSourceControlModule.h"
+#include "PlasticSourceControlUtils.h"
+
+#include "SourceControlOperations.h"
+#include "Fonts/SlateFontInfo.h"
+#include "Misc/Paths.h"
+#include "Misc/App.h"
+#include "Misc/FileHelper.h"
+#include "Modules/ModuleManager.h"
+#include "Framework/Notifications/NotificationManager.h"
+#include "Styling/SlateTypes.h"
+#include "Widgets/Input/SEditableTextBox.h"
+#include "Widgets/Input/SMultiLineEditableTextBox.h"
+#include "Widgets/Input/SButton.h"
+#include "Widgets/Input/SCheckBox.h"
+#include "Widgets/Layout/SBorder.h"
+#include "Widgets/Layout/SSeparator.h"
+#include "Widgets/Notifications/SNotificationList.h"
+#include "Widgets/SBoxPanel.h"
+#include "Widgets/Text/STextBlock.h"
+#include "EditorStyleSet.h"
+
+#define LOCTEXT_NAMESPACE "SPlasticSourceControlSettings"
+
+void SPlasticSourceControlSettings::Construct(const FArguments& InArgs)
+{
+	FSlateFontInfo Font = FEditorStyle::GetFontStyle(TEXT("SourceControl.LoginWindow.Font"));
+
+	bAutoCreateIgnoreFile = true;
+	bAutoInitialCommit = true;
+
+	InitialCommitMessage = LOCTEXT("InitialCommitMessage", "Initial checkin");
+	ServerUrl = FText::FromString(TEXT("localhost:8087"));
+	if (FApp::HasProjectName())
+	{
+		WorkspaceName = FText::FromString(FApp::GetProjectName());
+		RepositoryName = WorkspaceName;
+	}
+
+	ChildSlot
+	[
+		SNew(SBorder)
+		.BorderImage( FEditorStyle::GetBrush("DetailsView.CategoryBottom"))
+		.Padding(FMargin(0.0f, 3.0f, 0.0f, 0.0f))
+		[
+			SNew(SVerticalBox)
+			// Path to the Plastic SCM binary
+			+SVerticalBox::Slot()
+			.AutoHeight()
+			.Padding(2.0f)
+			.VAlign(VAlign_Center)
+			[
+				SNew(SHorizontalBox)
+				.ToolTipText(LOCTEXT("BinaryPathLabel_Tooltip", "Path to the Plastic SCM cli binary"))
+				+SHorizontalBox::Slot()
+				.FillWidth(1.0f)
+				[
+					SNew(STextBlock)
+					.Text(LOCTEXT("BinaryPathLabel", "Plastic SCM Path"))
+					.Font(Font)
+				]
+				+SHorizontalBox::Slot()
+				.FillWidth(2.0f)
+				[
+					SNew(SEditableTextBox)
+					.Text(this, &SPlasticSourceControlSettings::GetBinaryPathText)
+					.HintText(LOCTEXT("BinaryPathLabel", "Path to the Plastic SCM binary"))
+					.OnTextCommitted(this, &SPlasticSourceControlSettings::OnBinaryPathTextCommited)
+					.Font(Font)
+				]
+			]
+			// Root of the workspace
+			+SVerticalBox::Slot()
+			.AutoHeight()
+			.Padding(2.0f)
+			.VAlign(VAlign_Center)
+			[
+				SNew(SHorizontalBox)
+				.ToolTipText(LOCTEXT("WorkspaceRootLabel_Tooltip", "Path to the root of the Plastic SCM workspace"))
+				+SHorizontalBox::Slot()
+				.FillWidth(1.0f)
+				[
+					SNew(STextBlock)
+					.Text(LOCTEXT("WorkspaceRootLabel", "Root of the workspace"))
+					.Font(Font)
+				]
+				+SHorizontalBox::Slot()
+				.FillWidth(2.0f)
+				[
+					SNew(STextBlock)
+					.Text(this, &SPlasticSourceControlSettings::GetPathToWorkspaceRoot)
+					.Font(Font)
+				]
+			]
+			// User Name
+			+SVerticalBox::Slot()
+			.AutoHeight()
+			.Padding(2.0f)
+			.VAlign(VAlign_Center)
+			[
+				SNew(SHorizontalBox)
+				.ToolTipText(LOCTEXT("PlasticUserName_Tooltip", "User name configured for the Plastic SCM workspace"))
+				+SHorizontalBox::Slot()
+				.FillWidth(1.0f)
+				[
+					SNew(STextBlock)
+					.Text(LOCTEXT("PlasticUserName", "User Name"))
+					.Font(Font)
+				]
+				+SHorizontalBox::Slot()
+				.FillWidth(2.0f)
+				[
+					SNew(STextBlock)
+					.Text(this, &SPlasticSourceControlSettings::GetUserName)
+					.Font(Font)
+				]
+			]
+			// Separator
+			+SVerticalBox::Slot()
+			.AutoHeight()
+			.Padding(2.0f)
+			.VAlign(VAlign_Center)
+			[
+				SNew(SSeparator)
+			]
+			// Explanation text
+			+SVerticalBox::Slot()
+			.FillHeight(1.0f)
+			.Padding(2.0f)
+			.VAlign(VAlign_Center)
+			[
+				SNew(SHorizontalBox)
+				.Visibility(this, &SPlasticSourceControlSettings::CanInitializePlasticWorkspace)
+				.ToolTipText(LOCTEXT("WorkspaceNotFound_Tooltip", "No Workspace found at the level or above the current Project"))
+				+SHorizontalBox::Slot()
+				.FillWidth(1.0f)
+				.HAlign(HAlign_Center)
+				[
+					SNew(STextBlock)
+					.Text(LOCTEXT("WorkspaceNotFound", "Current Project is not contained in a Plastic SCM Workspace. Fill the form below to initialize a new Workspace."))
+					.Font(Font)
+				]
+			]
+			// Workspace and Repository Name
+			+SVerticalBox::Slot()
+			.AutoHeight()
+			.Padding(2.0f)
+			.VAlign(VAlign_Center)
+			[
+				SNew(SHorizontalBox)
+				.Visibility(this, &SPlasticSourceControlSettings::CanInitializePlasticWorkspace)
+				+SHorizontalBox::Slot()
+				.FillWidth(1.0f)
+				[
+					SNew(STextBlock)
+					.Text(LOCTEXT("WorkspaceRepositoryName", "Workspace and Repository Name"))
+					.ToolTipText(LOCTEXT("WorkspaceRepositoryName_Tooltip", "Enter the Name of the new Workspace and Repository to create or use"))
+					.Font(Font)
+				]
+				+SHorizontalBox::Slot()
+				.FillWidth(1.0f)
+				.Padding(2.0f, 0.0f, 0.0f, 0.0f)
+				[
+					SNew(SEditableTextBox)
+					.Text(this, &SPlasticSourceControlSettings::GetWorkspaceName)
+					.ToolTipText(LOCTEXT("WorkspaceName_Tooltip", "Enter the Name of the new Workspace to create"))
+					.HintText(LOCTEXT("WorkspaceName_Hint", "Name of the Workspace to create"))
+					.OnTextCommitted(this, &SPlasticSourceControlSettings::OnWorkspaceNameCommited)
+					.Font(Font)
+				]
+				+SHorizontalBox::Slot()
+				.FillWidth(1.0f)
+				.Padding(2.0f, 0.0f, 0.0f, 0.0f)
+				[
+					SNew(SEditableTextBox)
+					.Text(this, &SPlasticSourceControlSettings::GetRepositoryName)
+					.ToolTipText(LOCTEXT("RepositoryName_Tooltip", "Enter the Name of the Repository to use or create"))
+					.HintText(LOCTEXT("RepositoryName_Hint", "Name of the Repository to use or create"))
+					.OnTextCommitted(this, &SPlasticSourceControlSettings::OnRepositoryNameCommited)
+					.Font(Font)
+				]
+			]
+			// Server URL address:port
+			+SVerticalBox::Slot()
+			.AutoHeight()
+			.Padding(2.0f)
+			.VAlign(VAlign_Center)
+			[
+				SNew(SHorizontalBox)
+				.Visibility(this, &SPlasticSourceControlSettings::CanInitializePlasticWorkspace)
+				.ToolTipText(LOCTEXT("ServerUrl_Tooltip", "Enter the Server URL in the form address:port (localhost:8087)"))
+				+SHorizontalBox::Slot()
+				.FillWidth(1.0f)
+				[
+					SNew(STextBlock)
+					.Text(LOCTEXT("ServerUrl", "Server URL address:port"))
+					.Font(Font)
+				]
+				+SHorizontalBox::Slot()
+				.FillWidth(2.0f)
+				[
+					SNew(SEditableTextBox)
+					.Text(this, &SPlasticSourceControlSettings::GetServerUrl)
+					.HintText(LOCTEXT("ServerUrl", "Enter the Server URL"))
+					.OnTextCommitted(this, &SPlasticSourceControlSettings::OnServerUrlCommited)
+					.Font(Font)
+				]
+			]
+			// Option to add a 'ignore.conf' file at Workspace creation time
+			+SVerticalBox::Slot()
+			.AutoHeight()
+			.Padding(2.0f)
+			.VAlign(VAlign_Center)
+			[
+				SNew(SHorizontalBox)
+				.Visibility(this, &SPlasticSourceControlSettings::CanInitializePlasticWorkspace)
+				.ToolTipText(LOCTEXT("CreateIgnoreFile_Tooltip", "Create and add a standard 'ignore.conf' file"))
+				+SHorizontalBox::Slot()
+				.FillWidth(0.1f)
+				[
+					SNew(SCheckBox)
+					.IsChecked(ECheckBoxState::Checked)
+					.OnCheckStateChanged(this, &SPlasticSourceControlSettings::OnCheckedCreateIgnoreFile)
+				]
+				+SHorizontalBox::Slot()
+				.FillWidth(2.9f)
+				.VAlign(VAlign_Center)
+				[
+					SNew(STextBlock)
+					.Text(LOCTEXT("CreateIgnoreFile", "Add a ignore.conf file"))
+					.Font(Font)
+				]
+			]
+			// Option to Make the initial Plastic SCM checkin
+			+SVerticalBox::Slot()
+			.AutoHeight()
+			.Padding(2.0f)
+			.VAlign(VAlign_Center)
+			[
+				SNew(SHorizontalBox)
+				.Visibility(this, &SPlasticSourceControlSettings::CanInitializePlasticWorkspace)
+				.ToolTipText(LOCTEXT("InitialCommit_Tooltip", "Make the initial Plastic SCM checkin"))
+				+SHorizontalBox::Slot()
+				.FillWidth(0.1f)
+				[
+					SNew(SCheckBox)
+					.IsChecked(ECheckBoxState::Checked)
+					.OnCheckStateChanged(this, &SPlasticSourceControlSettings::OnCheckedInitialCommit)
+				]
+				+SHorizontalBox::Slot()
+				.FillWidth(0.9f)
+				.VAlign(VAlign_Center)
+				[
+					SNew(STextBlock)
+					.Text(LOCTEXT("InitialCommit", "Make the initial Plastic SCM Checkin"))
+					.Font(Font)
+				]
+				+SHorizontalBox::Slot()
+				.FillWidth(2.0f)
+				.Padding(2.0f)
+				[
+					SNew(SMultiLineEditableTextBox)
+					.Text(this, &SPlasticSourceControlSettings::GetInitialCommitMessage)
+					.HintText(LOCTEXT("InitialCommitMessage_Hint", "Message for the initial checkin"))
+					.OnTextCommitted(this, &SPlasticSourceControlSettings::OnInitialCommitMessageCommited)
+					.Font(Font)
+				]
+			]
+			// Option to enable the Update Status operation at Editor Startup
+			+SVerticalBox::Slot()
+			.AutoHeight()
+			.Padding(2.0f)
+			.VAlign(VAlign_Center)
+			[
+				SNew(SHorizontalBox)
+				.ToolTipText(LOCTEXT("UpdateStatusAtStartup_Tooltip", "Run an asynchronous Update Status at Editor startup (can be slow)."))
+				+SHorizontalBox::Slot()
+				.FillWidth(0.1f)
+				[
+					SNew(SCheckBox)
+					.IsChecked(SPlasticSourceControlSettings::IsUpdateStatusAtStartupChecked())
+					.OnCheckStateChanged(this, &SPlasticSourceControlSettings::OnCheckedUpdateStatusAtStartup)
+				]
+				+SHorizontalBox::Slot()
+				.FillWidth(2.9f)
+				.VAlign(VAlign_Center)
+				[
+					SNew(STextBlock)
+					.Text(LOCTEXT("UpdateStatusAtStartup", "Update workspace Status at Editor startup"))
+					.Font(Font)
+				]
+			]
+			// Button to create a new Workspace
+			+SVerticalBox::Slot()
+			.FillHeight(2.5f)
+			.Padding(4.0f)
+			.VAlign(VAlign_Center)
+			[
+				SNew(SHorizontalBox)
+				.Visibility(this, &SPlasticSourceControlSettings::CanInitializePlasticWorkspace)
+				+SHorizontalBox::Slot()
+				.FillWidth(1.0f)
+				[
+					SNew(SButton)
+					.IsEnabled(this, &SPlasticSourceControlSettings::IsReadyToInitializePlasticWorkspace)
+					.Text(LOCTEXT("PlasticInitWorkspace", "Create a new Plastic SCM workspace for the current project"))
+					.ToolTipText(LOCTEXT("PlasticInitWorkspace_Tooltip", "Create and initialize a new Plastic SCM workspace and repository for the current project"))
+					.OnClicked(this, &SPlasticSourceControlSettings::OnClickedInitializePlasticWorkspace)
+					.HAlign(HAlign_Center)
+					.ContentPadding(6)
+				]
+			]
+			// Button to add a 'ignore.conf' file on an existing Workspace
+			+SVerticalBox::Slot()
+			.FillHeight(2.0f)
+			.Padding(2.0f)
+			.VAlign(VAlign_Center)
+			[
+				SNew(SHorizontalBox)
+				.Visibility(this, &SPlasticSourceControlSettings::CanAddIgnoreFile)
+				+SHorizontalBox::Slot()
+				.FillWidth(1.0f)
+				[
+					SNew(SButton)
+					.Text(LOCTEXT("CreateIgnoreFile", "Add a ignore.conf file"))
+					.ToolTipText(LOCTEXT("CreateIgnoreFile_Tooltip", "Create and add a standard 'ignore.conf' file"))
+					.OnClicked(this, &SPlasticSourceControlSettings::OnClickedAddIgnoreFile)
+					.HAlign(HAlign_Center)
+				]
+			]
+		]
+	];
+}
+
+FText SPlasticSourceControlSettings::GetBinaryPathText() const
+{
+	const FPlasticSourceControlModule& PlasticSourceControl = FModuleManager::LoadModuleChecked<FPlasticSourceControlModule>("PlasticSourceControl");
+	return FText::FromString(PlasticSourceControl.AccessSettings().GetBinaryPath());
+}
+
+void SPlasticSourceControlSettings::OnBinaryPathTextCommited(const FText& InText, ETextCommit::Type InCommitType) const
+{
+	FPlasticSourceControlModule& PlasticSourceControl = FModuleManager::LoadModuleChecked<FPlasticSourceControlModule>("PlasticSourceControl");
+	const bool bChanged = PlasticSourceControl.AccessSettings().SetBinaryPath(InText.ToString());
+	if(bChanged)
+	{
+		// Re-Check provided Plastic binary path for each change
+		PlasticSourceControl.GetProvider().CheckPlasticAvailability();
+		if (PlasticSourceControl.GetProvider().IsPlasticAvailable())
+		{
+			PlasticSourceControl.SaveSettings();
+		}
+	}
+}
+
+FText SPlasticSourceControlSettings::GetPathToWorkspaceRoot() const
+{
+   const FPlasticSourceControlModule& PlasticSourceControl = FModuleManager::LoadModuleChecked<FPlasticSourceControlModule>("PlasticSourceControl");
+	return FText::FromString(PlasticSourceControl.GetProvider().GetPathToWorkspaceRoot());
+}
+
+FText SPlasticSourceControlSettings::GetUserName() const
+{
+   const FPlasticSourceControlModule& PlasticSourceControl = FModuleManager::LoadModuleChecked<FPlasticSourceControlModule>("PlasticSourceControl");
+	return FText::FromString(PlasticSourceControl.GetProvider().GetUserName());
+}
+
+
+EVisibility SPlasticSourceControlSettings::CanInitializePlasticWorkspace() const
+{
+   const FPlasticSourceControlModule& PlasticSourceControl = FModuleManager::LoadModuleChecked<FPlasticSourceControlModule>("PlasticSourceControl");
+	const bool bPlasticAvailable = PlasticSourceControl.GetProvider().IsPlasticAvailable();
+	const bool bPlasticWorkspaceFound = PlasticSourceControl.GetProvider().IsWorkspaceFound();
+	return (bPlasticAvailable && !bPlasticWorkspaceFound) ? EVisibility::Visible : EVisibility::Collapsed;
+}
+
+bool SPlasticSourceControlSettings::IsReadyToInitializePlasticWorkspace() const
+{
+	// Workspace Name cannot be left empty
+	const bool bWorkspaceNameOk = !WorkspaceName.IsEmpty();
+	// RepositoryName and ServerUrl should also be filled
+	const bool bRepositoryNameOk = !RepositoryName.IsEmpty() && !ServerUrl.IsEmpty();
+	// If Initial Commit is requested, checkin message cannot be empty
+	const bool bInitialCommitOk = (!bAutoInitialCommit || !InitialCommitMessage.IsEmpty());
+	return bWorkspaceNameOk && bRepositoryNameOk && bInitialCommitOk;
+}
+
+
+void SPlasticSourceControlSettings::OnWorkspaceNameCommited(const FText& InText, ETextCommit::Type InCommitType)
+{
+	WorkspaceName = InText;
+}
+FText SPlasticSourceControlSettings::GetWorkspaceName() const
+{
+	return WorkspaceName;
+}
+
+void SPlasticSourceControlSettings::OnRepositoryNameCommited(const FText& InText, ETextCommit::Type InCommitType)
+{
+	RepositoryName = InText;
+}
+FText SPlasticSourceControlSettings::GetRepositoryName() const
+{
+	return RepositoryName;
+}
+
+void SPlasticSourceControlSettings::OnServerUrlCommited(const FText& InText, ETextCommit::Type InCommitType)
+{
+	ServerUrl = InText;
+}
+FText SPlasticSourceControlSettings::GetServerUrl() const
+{
+	return ServerUrl;
+}
+
+void SPlasticSourceControlSettings::OnCheckedCreateIgnoreFile(ECheckBoxState NewCheckedState)
+{
+	bAutoCreateIgnoreFile = (NewCheckedState == ECheckBoxState::Checked);
+}
+
+void SPlasticSourceControlSettings::OnCheckedInitialCommit(ECheckBoxState NewCheckedState)
+{
+	bAutoInitialCommit = (NewCheckedState == ECheckBoxState::Checked);
+}
+
+void SPlasticSourceControlSettings::OnInitialCommitMessageCommited(const FText& InText, ETextCommit::Type InCommitType)
+{
+	InitialCommitMessage = InText;
+}
+
+FText SPlasticSourceControlSettings::GetInitialCommitMessage() const
+{
+	return InitialCommitMessage;
+}
+
+
+FReply SPlasticSourceControlSettings::OnClickedInitializePlasticWorkspace()
+{
+	UE_LOG(LogSourceControl, Log, TEXT("InitializePlasticWorkspace(%s, %s, %s) CreateIgnore=%d Commit=%d"),
+		*WorkspaceName.ToString(), *RepositoryName.ToString(), *ServerUrl.ToString(), bAutoCreateIgnoreFile, bAutoInitialCommit);
+
+	// 1.a. Create a repository (if not already existing) and a workspace: launch an asynchronous MakeWorkspace operation
+	LaunchMakeWorkspaceOperation();
+
+	return FReply::Handled();
+}
+
+
+/// 1. Create a repository (if not already existing) and a workspace
+void SPlasticSourceControlSettings::LaunchMakeWorkspaceOperation()
+{
+	TSharedRef<FPlasticMakeWorkspace, ESPMode::ThreadSafe> MakeWorkspaceOperation = ISourceControlOperation::Create<FPlasticMakeWorkspace>();
+	MakeWorkspaceOperation->WorkspaceName = WorkspaceName.ToString();
+	MakeWorkspaceOperation->RepositoryName = RepositoryName.ToString();
+	MakeWorkspaceOperation->ServerUrl = ServerUrl.ToString();
+
+	FPlasticSourceControlModule& PlasticSourceControl = FModuleManager::LoadModuleChecked<FPlasticSourceControlModule>("PlasticSourceControl");
+	ECommandResult::Type Result = PlasticSourceControl.GetProvider().Execute(MakeWorkspaceOperation, TArray<FString>(), EConcurrency::Asynchronous, FSourceControlOperationComplete::CreateSP(this, &SPlasticSourceControlSettings::OnSourceControlOperationComplete));
+	if (Result == ECommandResult::Succeeded)
+	{
+		DisplayInProgressNotification(MakeWorkspaceOperation->GetInProgressString());
+	}
+	else
+	{
+		DisplayFailureNotification(MakeWorkspaceOperation->GetName());
+	}
+}
+
+/// 2. Add all project files to to Source Control (.uproject, Config/, Content/, Source/ files and ignore.conf if any)
+void SPlasticSourceControlSettings::LaunchMarkForAddOperation()
+{
+	TSharedRef<FMarkForAdd, ESPMode::ThreadSafe> MarkForAddOperation = ISourceControlOperation::Create<FMarkForAdd>();
+	FPlasticSourceControlModule& PlasticSourceControl = FModuleManager::LoadModuleChecked<FPlasticSourceControlModule>("PlasticSourceControl");
+
+	// 1.b. Check the new workspace status to enable connection
+	PlasticSourceControl.GetProvider().CheckPlasticAvailability();
+
+	if (PlasticSourceControl.GetProvider().IsWorkspaceFound())
+	{
+		if (bAutoCreateIgnoreFile)
+		{
+			// 1.c Create a standard "ignore.conf" file with common patterns for a typical Blueprint & C++ project
+			CreateIgnoreFile();
+		}
+		// 2. Add all project files to to Source Control (.uproject, Config/, Content/, Source/ files and ignore.conf if any)
+		const TArray<FString> ProjectFiles = GetProjectFiles();
+		ECommandResult::Type Result = PlasticSourceControl.GetProvider().Execute(MarkForAddOperation, ProjectFiles, EConcurrency::Asynchronous, FSourceControlOperationComplete::CreateSP(this, &SPlasticSourceControlSettings::OnSourceControlOperationComplete));
+		if (Result == ECommandResult::Succeeded)
+		{
+			DisplayInProgressNotification(MarkForAddOperation->GetInProgressString());
+		}
+		else
+		{
+			DisplayFailureNotification(MarkForAddOperation->GetName());
+		}
+	}
+	else
+	{
+		DisplayFailureNotification(MarkForAddOperation->GetName());
+	}
+}
+
+/// 3. Launch an asynchronous "CheckIn" operation and start another ongoing notification
+void SPlasticSourceControlSettings::LaunchCheckInOperation()
+{
+	TSharedRef<FCheckIn, ESPMode::ThreadSafe> CheckInOperation = ISourceControlOperation::Create<FCheckIn>();
+	CheckInOperation->SetDescription(InitialCommitMessage);
+	FPlasticSourceControlModule& PlasticSourceControl = FModuleManager::LoadModuleChecked<FPlasticSourceControlModule>("PlasticSourceControl");
+	const TArray<FString> ProjectFiles = GetProjectFiles();
+	ECommandResult::Type Result = PlasticSourceControl.GetProvider().Execute(CheckInOperation, ProjectFiles, EConcurrency::Asynchronous, FSourceControlOperationComplete::CreateSP(this, &SPlasticSourceControlSettings::OnSourceControlOperationComplete));
+	if (Result == ECommandResult::Succeeded)
+	{
+		DisplayInProgressNotification(CheckInOperation->GetInProgressString());
+	}
+	else
+	{
+		DisplayFailureNotification(CheckInOperation->GetName());
+	}
+}
+
+/// Delegate called when a source control operation has completed: launch the next one and manage notifications
+void SPlasticSourceControlSettings::OnSourceControlOperationComplete(const FSourceControlOperationRef& InOperation, ECommandResult::Type InResult)
+{
+	RemoveInProgressNotification();
+
+	// Report result with a notification
+	if (InResult == ECommandResult::Succeeded)
+	{
+		DisplaySuccessNotification(InOperation->GetName());
+	}
+	else
+	{
+		DisplayFailureNotification(InOperation->GetName());
+	}
+
+	// Launch the following asynchrounous operation
+	if ((InOperation->GetName() == "MakeWorkspace") && (InResult == ECommandResult::Succeeded) && bAutoInitialCommit)
+	{
+		// 2. Add .uproject, Config/, Content/ and Source/ files (and ignore.conf if any)
+		LaunchMarkForAddOperation();
+	}
+	else if ((InOperation->GetName() == "MarkForAdd") && (InResult == ECommandResult::Succeeded) && bAutoInitialCommit)
+	{
+		// 3. optional initial Asynchronous commit with custom message: launch a "CheckIn" Operation
+		LaunchCheckInOperation();
+	}
+}
+
+// Display an ongoing notification during the whole operation
+void SPlasticSourceControlSettings::DisplayInProgressNotification(const FText& InOperationInProgressString)
+{
+	if (!OperationInProgressNotification.IsValid())
+	{
+		FNotificationInfo Info(InOperationInProgressString);
+		Info.bFireAndForget = false;
+		Info.ExpireDuration = 0.0f;
+		Info.FadeOutDuration = 1.0f;
+		OperationInProgressNotification = FSlateNotificationManager::Get().AddNotification(Info);
+		if (OperationInProgressNotification.IsValid())
+		{
+			OperationInProgressNotification.Pin()->SetCompletionState(SNotificationItem::CS_Pending);
+		}
+	}
+}
+
+// Remove the ongoing notification at the end of the operation
+void SPlasticSourceControlSettings::RemoveInProgressNotification()
+{
+	if (OperationInProgressNotification.IsValid())
+	{
+		OperationInProgressNotification.Pin()->ExpireAndFadeout();
+		OperationInProgressNotification.Reset();
+	}
+}
+
+// Display a temporary success notification at the end of the operation
+void SPlasticSourceControlSettings::DisplaySuccessNotification(const FName& InOperationName)
+{
+	const FText NotificationText = FText::Format(LOCTEXT("InitWorkspace_Success", "{0} operation was successful!"), FText::FromName(InOperationName));
+	FNotificationInfo Info(NotificationText);
+	Info.bUseSuccessFailIcons = true;
+	Info.Image = FEditorStyle::GetBrush(TEXT("NotificationList.SuccessImage"));
+	FSlateNotificationManager::Get().AddNotification(Info);
+	// @todo temporary debug log
+	UE_LOG(LogSourceControl, Log, TEXT("%s"), *NotificationText.ToString());
+}
+
+// Display a temporary failure notification at the end of the operation
+void SPlasticSourceControlSettings::DisplayFailureNotification(const FName& InOperationName)
+{
+	const FText NotificationText = FText::Format(LOCTEXT("InitWorkspace_Failure", "Error: {0} operation failed!"), FText::FromName(InOperationName));
+	FNotificationInfo Info(NotificationText);
+	Info.ExpireDuration = 8.0f;
+	FSlateNotificationManager::Get().AddNotification(Info);
+	// @todo temporary debug log
+	UE_LOG(LogSourceControl, Error, TEXT("%s"), *NotificationText.ToString());
+}
+
+/** Delegate to check for presence of a Plastic ignore.conf file to an existing Plastic SCM workspace */
+EVisibility SPlasticSourceControlSettings::CanAddIgnoreFile() const
+{
+   const FPlasticSourceControlModule& PlasticSourceControl = FModuleManager::LoadModuleChecked<FPlasticSourceControlModule>("PlasticSourceControl");
+	const bool bPlasticWorkspaceFound = PlasticSourceControl.GetProvider().IsWorkspaceFound();
+	const bool bIgnoreFileFound = FPaths::FileExists(GetIgnoreFileName());
+	return (bPlasticWorkspaceFound && !bIgnoreFileFound) ? EVisibility::Visible : EVisibility::Collapsed;
+}
+
+/** Delegate to add a Plastic ignore.conf file to an existing Plastic SCM workspace */
+FReply SPlasticSourceControlSettings::OnClickedAddIgnoreFile() const
+{
+	if (CreateIgnoreFile())
+	{
+		// Add ignore.conf to Plastic SCM
+		TArray<FString> InfoMessages;
+		TArray<FString> ErrorMessages;
+		TArray<FString> Parameters;
+		Parameters.Add(TEXT("-R"));
+		TArray<FString> Files;
+		Files.Add(TEXT("ignore.conf"));
+		PlasticSourceControlUtils::RunCommand(TEXT("add"), Parameters, Files, EConcurrency::Synchronous, InfoMessages, ErrorMessages);
+	}
+	return FReply::Handled();
+}
+
+void SPlasticSourceControlSettings::OnCheckedUpdateStatusAtStartup(ECheckBoxState NewCheckedState)
+{
+	FPlasticSourceControlModule& PlasticSourceControl = FModuleManager::GetModuleChecked<FPlasticSourceControlModule>("PlasticSourceControl");
+	PlasticSourceControl.AccessSettings().SetUpdateStatusAtStartup(NewCheckedState == ECheckBoxState::Checked);
+	PlasticSourceControl.AccessSettings().SaveSettings();
+}
+
+bool SPlasticSourceControlSettings::UpdateStatusAtStartup() const
+{
+	const FPlasticSourceControlModule& PlasticSourceControl = FModuleManager::GetModuleChecked<FPlasticSourceControlModule>("PlasticSourceControl");
+	return PlasticSourceControl.AccessSettings().UpdateStatusAtStartup();
+}
+
+ECheckBoxState SPlasticSourceControlSettings::IsUpdateStatusAtStartupChecked() const
+{
+	return (UpdateStatusAtStartup() ? ECheckBoxState::Checked : ECheckBoxState::Unchecked);
+}
+
+/** Path to the "ignore.conf" file */
+const FString& SPlasticSourceControlSettings::GetIgnoreFileName() const
+{
+	static const FString PathToProjectDir = FPaths::ConvertRelativePathToFull(FPaths::ProjectDir());
+	static const FString IgnoreFileName = FPaths::Combine(*PathToProjectDir, TEXT("ignore.conf"));
+	return IgnoreFileName;
+}
+
+/** Create a standard "ignore.conf" file with common patterns for a typical Blueprint & C++ project */
+bool SPlasticSourceControlSettings::CreateIgnoreFile() const
+{
+	const FString IgnoreFileContent = TEXT("Binaries\nBuild\nDerivedDataCache\nIntermediate\nSaved\n.vscode\n.vs\n*.VC.db\n*.opensdf\n*.opendb\n*.sdf\n*.sln\n*.suo\n*.xcodeproj\n*.xcworkspace");
+	return FFileHelper::SaveStringToFile(IgnoreFileContent, *GetIgnoreFileName(), FFileHelper::EEncodingOptions::ForceUTF8WithoutBOM);
+}
+
+/** List of files to add to Source Control (.uproject, Config/, Content/, Source/ files and ignore.conf if any) */
+TArray<FString> SPlasticSourceControlSettings::GetProjectFiles() const
+{
+	TArray<FString> ProjectFiles;
+	ProjectFiles.Add(FPaths::ConvertRelativePathToFull(FPaths::GetProjectFilePath()));
+	ProjectFiles.Add(FPaths::ConvertRelativePathToFull(FPaths::ProjectConfigDir()));
+	ProjectFiles.Add(FPaths::ConvertRelativePathToFull(FPaths::ProjectContentDir()));
+	if (FPaths::DirectoryExists(FPaths::GameSourceDir()))
+	{
+		ProjectFiles.Add(FPaths::ConvertRelativePathToFull(FPaths::GameSourceDir()));
+	}
+	if (bAutoCreateIgnoreFile)
+	{
+		ProjectFiles.Add(GetIgnoreFileName());
+	}
+	return ProjectFiles;
+}
+
+#undef LOCTEXT_NAMESPACE
diff --git a/Engine/Plugins/Developer/PlasticSourceControl/Source/PlasticSourceControl/Private/SPlasticSourceControlSettings.h b/Engine/Plugins/Developer/PlasticSourceControl/Source/PlasticSourceControl/Private/SPlasticSourceControlSettings.h
new file mode 100644
--- /dev/null
+++ b/Engine/Plugins/Developer/PlasticSourceControl/Source/PlasticSourceControl/Private/SPlasticSourceControlSettings.h
@@ -0,0 +1,90 @@
+// Copyright 1998-2019 Epic Games, Inc. All Rights Reserved.
+
+#pragma once
+
+#include "CoreMinimal.h"
+#include "Widgets/SCompoundWidget.h"
+#include "Widgets/DeclarativeSyntaxSupport.h"
+#include "Layout/Visibility.h"
+#include "Input/Reply.h"
+#include "Styling/SlateTypes.h"
+
+#include "ISourceControlProvider.h"
+#include "ISourceControlOperation.h"
+
+class SPlasticSourceControlSettings : public SCompoundWidget
+{
+public:
+	
+	SLATE_BEGIN_ARGS(SPlasticSourceControlSettings) {}
+	
+	SLATE_END_ARGS()
+
+public:
+
+	void Construct(const FArguments& InArgs);
+
+private:
+
+	/** Delegate to get cm binary path from settings */
+	FText GetBinaryPathText() const;
+
+	/** Delegate to commit cm binary path to settings */
+	void OnBinaryPathTextCommited(const FText& InText, ETextCommit::Type InCommitType) const;
+
+	/** Delegate to get workspace root and user name from provider */
+	FText GetPathToWorkspaceRoot() const;
+	FText GetUserName() const;
+
+	/** Delegate to initialize a new Plastic workspace and repository */
+	EVisibility CanInitializePlasticWorkspace() const;
+	bool IsReadyToInitializePlasticWorkspace() const;
+	void OnWorkspaceNameCommited(const FText& InText, ETextCommit::Type InCommitType);
+	FText GetWorkspaceName() const;
+	FText WorkspaceName;
+	void OnRepositoryNameCommited(const FText& InText, ETextCommit::Type InCommitType);
+	FText GetRepositoryName() const;
+	FText RepositoryName;
+	void OnServerUrlCommited(const FText& InText, ETextCommit::Type InCommitType);
+	FText GetServerUrl() const;
+	FText ServerUrl;
+	void OnCheckedCreateIgnoreFile(ECheckBoxState NewCheckedState);
+	bool bAutoCreateIgnoreFile;
+
+	void OnCheckedUpdateStatusAtStartup(ECheckBoxState NewCheckedState);
+	ECheckBoxState IsUpdateStatusAtStartupChecked() const;
+	bool UpdateStatusAtStartup() const;
+
+	void OnCheckedInitialCommit(ECheckBoxState NewCheckedState);
+	bool bAutoInitialCommit;
+	void OnInitialCommitMessageCommited(const FText& InText, ETextCommit::Type InCommitType);
+	FText GetInitialCommitMessage() const;
+	FText InitialCommitMessage;
+
+	/** Launch initial asynchronous add and commit operations */
+	void LaunchMakeWorkspaceOperation();
+	void LaunchMarkForAddOperation();
+	void LaunchCheckInOperation();
+
+	/** Delegate called when a source control operation has completed */
+	void OnSourceControlOperationComplete(const FSourceControlOperationRef& InOperation, ECommandResult::Type InResult);
+
+	/** Asynchronous operation progress notifications */
+	TWeakPtr<class SNotificationItem> OperationInProgressNotification;
+
+	void DisplayInProgressNotification(const FText& InOperationInProgressString);
+	void RemoveInProgressNotification();
+	void DisplaySuccessNotification(const FName& InOperationName);
+	void DisplayFailureNotification(const FName& InOperationName);
+
+	FReply OnClickedInitializePlasticWorkspace();
+
+	/** Delegate to add a Plastic ignore.conf file to an existing Plastic workspace */
+	EVisibility CanAddIgnoreFile() const;
+	FReply OnClickedAddIgnoreFile() const;
+
+	const FString& GetIgnoreFileName() const;
+	bool CreateIgnoreFile() const;
+
+	TArray<FString> GetProjectFiles() const;
+};
