Fix compile errors with PlasticSourceControl
#rnx

diff --git a/Engine/Plugins/Developer/PlasticSourceControl/Source/PlasticSourceControl/Private/PlasticSourceControlMenu.cpp b/Engine/Plugins/Developer/PlasticSourceControl/Source/PlasticSourceControl/Private/PlasticSourceControlMenu.cpp
--- a/Engine/Plugins/Developer/PlasticSourceControl/Source/PlasticSourceControl/Private/PlasticSourceControlMenu.cpp
+++ b/Engine/Plugins/Developer/PlasticSourceControl/Source/PlasticSourceControl/Private/PlasticSourceControlMenu.cpp
@@ -156,10 +156,10 @@ void FPlasticSourceControlMenu::ReloadPackages(TArray<UPackage*>& InPackagesToRe
 	});
 
 	// Hot-reload the new packages...
-	PackageTools::ReloadPackages(InPackagesToReload);
+	UPackageTools::ReloadPackages(InPackagesToReload);
 
 	// Unload any deleted packages...
-	PackageTools::UnloadPackages(PackagesToUnload);
+	UPackageTools::UnloadPackages(PackagesToUnload);
 }
 
 void FPlasticSourceControlMenu::SyncProjectClicked()
diff --git a/Engine/Plugins/Developer/PlasticSourceControl/Source/PlasticSourceControl/Private/PlasticSourceControlModule.h b/Engine/Plugins/Developer/PlasticSourceControl/Source/PlasticSourceControl/Private/PlasticSourceControlModule.h
--- a/Engine/Plugins/Developer/PlasticSourceControl/Source/PlasticSourceControl/Private/PlasticSourceControlModule.h
+++ b/Engine/Plugins/Developer/PlasticSourceControl/Source/PlasticSourceControl/Private/PlasticSourceControlModule.h
@@ -4,6 +4,7 @@
 
 #include "CoreMinimal.h"
 #include "Modules/ModuleInterface.h"
+#include "Modules/ModuleManager.h"
 #include "PlasticSourceControlSettings.h"
 #include "PlasticSourceControlProvider.h"
 
diff --git a/Engine/Plugins/Developer/PlasticSourceControl/Source/PlasticSourceControl/Private/PlasticSourceControlOperations.cpp b/Engine/Plugins/Developer/PlasticSourceControl/Source/PlasticSourceControl/Private/PlasticSourceControlOperations.cpp
--- a/Engine/Plugins/Developer/PlasticSourceControl/Source/PlasticSourceControl/Private/PlasticSourceControlOperations.cpp
+++ b/Engine/Plugins/Developer/PlasticSourceControl/Source/PlasticSourceControl/Private/PlasticSourceControlOperations.cpp
@@ -12,6 +12,7 @@
 #include "HAL/FileManager.h"
 #include "Misc/Paths.h"
 #include "SourceControlOperations.h"
+#include "ISourceControlModule.h"
 
 #define LOCTEXT_NAMESPACE "PlasticSourceControl"
 
diff --git a/Engine/Plugins/Developer/PlasticSourceControl/Source/PlasticSourceControl/Private/PlasticSourceControlRevision.cpp b/Engine/Plugins/Developer/PlasticSourceControl/Source/PlasticSourceControl/Private/PlasticSourceControlRevision.cpp
--- a/Engine/Plugins/Developer/PlasticSourceControl/Source/PlasticSourceControl/Private/PlasticSourceControlRevision.cpp
+++ b/Engine/Plugins/Developer/PlasticSourceControl/Source/PlasticSourceControl/Private/PlasticSourceControlRevision.cpp
@@ -9,6 +9,7 @@
 
 #include "HAL/FileManager.h"
 #include "Misc/Paths.h"
+#include "ISourceControlModule.h"
 
 #define LOCTEXT_NAMESPACE "PlasticSourceControl"
 
diff --git a/Engine/Plugins/Developer/PlasticSourceControl/Source/PlasticSourceControl/Private/PlasticSourceControlState.cpp b/Engine/Plugins/Developer/PlasticSourceControl/Source/PlasticSourceControl/Private/PlasticSourceControlState.cpp
--- a/Engine/Plugins/Developer/PlasticSourceControl/Source/PlasticSourceControl/Private/PlasticSourceControlState.cpp
+++ b/Engine/Plugins/Developer/PlasticSourceControl/Source/PlasticSourceControl/Private/PlasticSourceControlState.cpp
@@ -1,6 +1,7 @@
 // Copyright 1998-2019 Epic Games, Inc. All Rights Reserved.
 
 #include "PlasticSourceControlState.h"
+#include "ISourceControlModule.h"
 
 #define LOCTEXT_NAMESPACE "PlasticSourceControl.State"
 
diff --git a/Engine/Plugins/Developer/PlasticSourceControl/Source/PlasticSourceControl/Private/PlasticSourceControlUtils.cpp b/Engine/Plugins/Developer/PlasticSourceControl/Source/PlasticSourceControl/Private/PlasticSourceControlUtils.cpp
--- a/Engine/Plugins/Developer/PlasticSourceControl/Source/PlasticSourceControl/Private/PlasticSourceControlUtils.cpp
+++ b/Engine/Plugins/Developer/PlasticSourceControl/Source/PlasticSourceControl/Private/PlasticSourceControlUtils.cpp
@@ -14,6 +14,7 @@
 #include "HAL/PlatformTime.h"
 #include "Modules/ModuleManager.h"
 #include "XmlParser.h"
+#include "ISourceControlModule.h"
 
 #if PLATFORM_LINUX
 #include <sys/ioctl.h>
diff --git a/Engine/Plugins/Developer/PlasticSourceControl/Source/PlasticSourceControl/Private/SPlasticSourceControlSettings.cpp b/Engine/Plugins/Developer/PlasticSourceControl/Source/PlasticSourceControl/Private/SPlasticSourceControlSettings.cpp
--- a/Engine/Plugins/Developer/PlasticSourceControl/Source/PlasticSourceControl/Private/SPlasticSourceControlSettings.cpp
+++ b/Engine/Plugins/Developer/PlasticSourceControl/Source/PlasticSourceControl/Private/SPlasticSourceControlSettings.cpp
@@ -6,6 +6,7 @@
 #include "PlasticSourceControlUtils.h"
 
 #include "SourceControlOperations.h"
+#include "ISourceControlModule.h"
 #include "Fonts/SlateFontInfo.h"
 #include "Misc/Paths.h"
 #include "Misc/App.h"
