Fix null pointer dereference in plastic source control
#rnx

diff --git a/Engine/Plugins/Developer/PlasticSourceControl/Source/PlasticSourceControl/Private/PlasticSourceControlRevision.cpp b/Engine/Plugins/Developer/PlasticSourceControl/Source/PlasticSourceControl/Private/PlasticSourceControlRevision.cpp
--- a/Engine/Plugins/Developer/PlasticSourceControl/Source/PlasticSourceControl/Private/PlasticSourceControlRevision.cpp
+++ b/Engine/Plugins/Developer/PlasticSourceControl/Source/PlasticSourceControl/Private/PlasticSourceControlRevision.cpp
@@ -37,12 +37,16 @@ bool FPlasticSourceControlRevision::Get( FString& InOutFilename ) const
 	{
 		bCommandSuccessful = true; // if the temp file already exists, reuse it directly
 	}
-	else
+	else if (State)
 	{
 		// Format the revision specification of the file, like revid:1230@repo@server:8087
 		const FString RevisionSpecification = FString::Printf(TEXT("revid:%d@%s"), RevisionId, *State->RepSpec);
 		bCommandSuccessful = PlasticSourceControlUtils::RunDumpToFile(PathToPlasticBinary, RevisionSpecification, InOutFilename);
 	}
+	else
+	{
+		bCommandSuccessful = false;
+	}
 	return bCommandSuccessful;
 }
 
