Removing old comment

#rnx

diff --git a/Engine/Plugins/Developer/PlasticSourceControl/Source/PlasticSourceControl/Private/PlasticSourceControlProvider.cpp b/Engine/Plugins/Developer/PlasticSourceControl/Source/PlasticSourceControl/Private/PlasticSourceControlProvider.cpp
--- a/Engine/Plugins/Developer/PlasticSourceControl/Source/PlasticSourceControl/Private/PlasticSourceControlProvider.cpp
+++ b/Engine/Plugins/Developer/PlasticSourceControl/Source/PlasticSourceControl/Private/PlasticSourceControlProvider.cpp
@@ -379,10 +379,6 @@ void FPlasticSourceControlProvider::Tick()
 TArray< TSharedRef<ISourceControlLabel> > FPlasticSourceControlProvider::GetLabels( const FString& InMatchingSpec ) const
 {
 	TArray< TSharedRef<ISourceControlLabel> > Tags;
-
-	// NOTE list labels. Called by CrashDebugHelper() (to remote debug Engine crash)
-	//					 and by SourceControlHelpers::AnnotateFile() (to add source file to report)
-	// Reserved for internal use by Epic Games with Perforce only
 	return Tags;
 }
 
