Remove UE4 references in source control, AssetTools, Persona

diff --git a/Engine/Plugins/Developer/GitSourceControl/Source/GitSourceControl/Private/GitSourceControlModule.h b/Engine/Plugins/Developer/GitSourceControl/Source/GitSourceControl/Private/GitSourceControlModule.h
--- a/Engine/Plugins/Developer/GitSourceControl/Source/GitSourceControl/Private/GitSourceControlModule.h
+++ b/Engine/Plugins/Developer/GitSourceControl/Source/GitSourceControl/Private/GitSourceControlModule.h
@@ -9,12 +9,12 @@
 
 /**
 
-UE4GitPlugin is a simple Git Source Control Plugin for Unreal Engine
+UEGitPlugin is a simple Git Source Control Plugin for Unreal Engine
 
 Written and contributed by [REDACTED_PERSON] ([REDACTED_EMAIL])
 
 ### Supported features
-- initialize a new Git local repository ('git init') to manage your UE4 Game Project.
+- initialize a new Git local repository ('git init') to manage your UE Game Project.
 - create an appropriate .gitignore file as part as initialization
 - can also make the initial commit
 - display status icons to show modified/added/deleted/untracked files
diff --git a/Engine/Plugins/Developer/GitSourceControl/Source/GitSourceControl/Private/GitSourceControlState.cpp b/Engine/Plugins/Developer/GitSourceControl/Source/GitSourceControl/Private/GitSourceControlState.cpp
--- a/Engine/Plugins/Developer/GitSourceControl/Source/GitSourceControl/Private/GitSourceControlState.cpp
+++ b/Engine/Plugins/Developer/GitSourceControl/Source/GitSourceControl/Private/GitSourceControlState.cpp
@@ -228,7 +228,7 @@ bool FGitSourceControlState::IsModified() const
 	// so for a clean "check-in" (commit) checked-out files unmodified should be removed from the changeset (the index)
 	// http://stackoverflow.com/questions/12357971/what-does-revert-unchanged-files-mean-in-perforce
 	//
-	// Thus, before check-in UE4 Editor call RevertUnchangedFiles() in PromptForCheckin() and CheckinFiles().
+	// Thus, before check-in UE Editor call RevertUnchangedFiles() in PromptForCheckin() and CheckinFiles().
 	//
 	// So here we must take care to enumerate all states that need to be commited,
 	// all other will be discarded :
diff --git a/Engine/Plugins/Developer/PerforceSourceControl/Source/PerforceSourceControl/Private/PerforceConnection.cpp b/Engine/Plugins/Developer/PerforceSourceControl/Source/PerforceSourceControl/Private/PerforceConnection.cpp
--- a/Engine/Plugins/Developer/PerforceSourceControl/Source/PerforceSourceControl/Private/PerforceConnection.cpp
+++ b/Engine/Plugins/Developer/PerforceSourceControl/Source/PerforceSourceControl/Private/PerforceConnection.cpp
@@ -425,7 +425,7 @@ bool FPerforceConnection::EnsureValidConnection(FString& InOutServerName, FStrin
 	FString NewClientSpecName = InOutWorkspaceName;
 
 	ClientApi TestP4;
-	TestP4.SetProg("UE4");
+	TestP4.SetProg("UE");
 	TestP4.SetProtocol("tag", "");
 	TestP4.SetProtocol("enableStreams", "");
 
@@ -844,7 +844,7 @@ void FPerforceConnection::EstablishConnection(const FPerforceConnectionInfo& InC
 
 	UE_LOG(LogSourceControl, Verbose, TEXT("Attempting P4 connection: %s/%s"), *InConnectionInfo.Port, *InConnectionInfo.UserName);
 
-	P4Client.SetProg("UE4");
+	P4Client.SetProg("UE");
 	P4Client.SetProtocol("tag", "");
 	P4Client.SetProtocol("enableStreams", "");
 
diff --git a/Engine/Plugins/Developer/PerforceSourceControl/Source/PerforceSourceControl/Private/SPerforceSourceControlSettings.cpp b/Engine/Plugins/Developer/PerforceSourceControl/Source/PerforceSourceControl/Private/SPerforceSourceControlSettings.cpp
--- a/Engine/Plugins/Developer/PerforceSourceControl/Source/PerforceSourceControl/Private/SPerforceSourceControlSettings.cpp
+++ b/Engine/Plugins/Developer/PerforceSourceControl/Source/PerforceSourceControl/Private/SPerforceSourceControlSettings.cpp
@@ -38,7 +38,7 @@ void SPerforceSourceControlSettings::Construct(const FArguments& InArgs)
 	{
 #if USE_P4_API
 		ClientApi TestP4;
-		TestP4.SetProg("UE4");
+		TestP4.SetProg("UE");
 		Error P4Error;
 		TestP4.Init(&P4Error);
 		PortName = ANSI_TO_TCHAR(TestP4.GetPort().Text());
diff --git a/Engine/Plugins/Developer/PlasticSourceControl/Source/PlasticSourceControl/Private/PlasticSourceControlModule.h b/Engine/Plugins/Developer/PlasticSourceControl/Source/PlasticSourceControl/Private/PlasticSourceControlModule.h
--- a/Engine/Plugins/Developer/PlasticSourceControl/Source/PlasticSourceControl/Private/PlasticSourceControlModule.h
+++ b/Engine/Plugins/Developer/PlasticSourceControl/Source/PlasticSourceControl/Private/PlasticSourceControlModule.h
@@ -9,7 +9,7 @@
 #include "PlasticSourceControlProvider.h"
 
 /**
- * UE4PlasticPlugin is a simple Plastic SCM Source Control Plugin for Unreal Engine 4
+ * UEPlasticPlugin is a simple Plastic SCM Source Control Plugin for Unreal Engine
  *
  * Written and contributed by [REDACTED_PERSON] ([REDACTED_EMAIL])
  */
diff --git a/Engine/Plugins/Developer/PlasticSourceControl/Source/PlasticSourceControl/Private/PlasticSourceControlProvider.h b/Engine/Plugins/Developer/PlasticSourceControl/Source/PlasticSourceControl/Private/PlasticSourceControlProvider.h
--- a/Engine/Plugins/Developer/PlasticSourceControl/Source/PlasticSourceControl/Private/PlasticSourceControlProvider.h
+++ b/Engine/Plugins/Developer/PlasticSourceControl/Source/PlasticSourceControl/Private/PlasticSourceControlProvider.h
@@ -29,9 +29,9 @@ public:
 	virtual bool IsEnabled() const override;
 	virtual bool IsAvailable() const override;
 	virtual const FName& GetName(void) const override;
-	virtual bool QueryStateBranchConfig(const FString& ConfigSrc, const FString& ConfigDest) /* override UE4.20 */ { return false; }
-	virtual void RegisterStateBranches(const TArray<FString>& BranchNames, const FString& ContentRoot) /* override UE4.20 */ {}
-	virtual int32 GetStateBranchIndex(const FString& InBranchName) const /* override UE4.20 */ { return INDEX_NONE; }
+	virtual bool QueryStateBranchConfig(const FString& ConfigSrc, const FString& ConfigDest) { return false; }
+	virtual void RegisterStateBranches(const TArray<FString>& BranchNames, const FString& ContentRoot) {}
+	virtual int32 GetStateBranchIndex(const FString& InBranchName) const { return INDEX_NONE; }
 	virtual ECommandResult::Type GetState( const TArray<FString>& InFiles, TArray<FSourceControlStateRef>& OutState, EStateCacheUsage::Type InStateCacheUsage ) override;
 	virtual ECommandResult::Type GetState(const TArray<FSourceControlChangelistRef>& InChangelists, TArray<FSourceControlChangelistStateRef>& OutState, EStateCacheUsage::Type InStateCacheUsage) override;
 	virtual TArray<FSourceControlStateRef> GetCachedStateByPredicate(TFunctionRef<bool(const FSourceControlStateRef&)> Predicate) const override;
diff --git a/Engine/Plugins/Developer/PlasticSourceControl/Source/PlasticSourceControl/Private/PlasticSourceControlState.cpp b/Engine/Plugins/Developer/PlasticSourceControl/Source/PlasticSourceControl/Private/PlasticSourceControlState.cpp
--- a/Engine/Plugins/Developer/PlasticSourceControl/Source/PlasticSourceControl/Private/PlasticSourceControlState.cpp
+++ b/Engine/Plugins/Developer/PlasticSourceControl/Source/PlasticSourceControl/Private/PlasticSourceControlState.cpp
@@ -378,7 +378,7 @@ bool FPlasticSourceControlState::IsModified() const
 {
 	// Warning: for a clean "check-in" (commit) checked-out files unmodified should be removed from the changeset (the index)
 	//
-	// Thus, before check-in UE4 Editor call RevertUnchangedFiles() in PromptForCheckin() and CheckinFiles().
+	// Thus, before check-in UE Editor call RevertUnchangedFiles() in PromptForCheckin() and CheckinFiles().
 	//
 	// So here we must take care to enumerate all states that need to be commited, all other will be discarded:
 	//  - Unknown
diff --git a/Engine/Plugins/Developer/PlasticSourceControl/Source/PlasticSourceControl/Private/PlasticSourceControlState.h b/Engine/Plugins/Developer/PlasticSourceControl/Source/PlasticSourceControl/Private/PlasticSourceControlState.h
--- a/Engine/Plugins/Developer/PlasticSourceControl/Source/PlasticSourceControl/Private/PlasticSourceControlState.h
+++ b/Engine/Plugins/Developer/PlasticSourceControl/Source/PlasticSourceControl/Private/PlasticSourceControlState.h
@@ -102,12 +102,12 @@ public:
 	virtual bool CanCheckout() const override;
 	virtual bool IsCheckedOut() const override;
 	virtual bool IsCheckedOutOther(FString* Who = nullptr) const override;
-	virtual bool IsCheckedOutInOtherBranch(const FString& CurrentBranch = FString()) const /* override UE4.20 */;
-	virtual bool IsModifiedInOtherBranch(const FString& CurrentBranch = FString()) const /* override UE4.20 */;
-	virtual bool IsCheckedOutOrModifiedInOtherBranch(const FString& CurrentBranch = FString()) const /* override UE4.20 */;
-	virtual TArray<FString> GetCheckedOutBranches() const /* override UE4.20 */;
-	virtual FString GetOtherUserBranchCheckedOuts() const /* override UE4.20 */;
-	virtual bool GetOtherBranchHeadModification(FString& HeadBranchOut, FString& ActionOut, int32& HeadChangeListOut) const /* override UE4.20 */;
+	virtual bool IsCheckedOutInOtherBranch(const FString& CurrentBranch = FString()) const;
+	virtual bool IsModifiedInOtherBranch(const FString& CurrentBranch = FString()) const;
+	virtual bool IsCheckedOutOrModifiedInOtherBranch(const FString& CurrentBranch = FString()) const;
+	virtual TArray<FString> GetCheckedOutBranches() const;
+	virtual FString GetOtherUserBranchCheckedOuts() const;
+	virtual bool GetOtherBranchHeadModification(FString& HeadBranchOut, FString& ActionOut, int32& HeadChangeListOut) const;
 	virtual bool IsCurrent() const override;
 	virtual bool IsSourceControlled() const override;
 	virtual bool IsAdded() const override;
@@ -119,7 +119,7 @@ public:
 	virtual bool CanAdd() const override;
 	virtual bool CanDelete() const override;
 	virtual bool IsConflicted() const override;
-	virtual bool CanRevert() const; // override; TODO UE 4.19 only - commented for backward compatibility
+	virtual bool CanRevert() const;
 
 public:
 	/** History of the item, if any */
diff --git a/Engine/Plugins/Developer/PlasticSourceControl/Source/PlasticSourceControl/Private/PlasticSourceControlUtils.cpp b/Engine/Plugins/Developer/PlasticSourceControl/Source/PlasticSourceControl/Private/PlasticSourceControlUtils.cpp
--- a/Engine/Plugins/Developer/PlasticSourceControl/Source/PlasticSourceControl/Private/PlasticSourceControlUtils.cpp
+++ b/Engine/Plugins/Developer/PlasticSourceControl/Source/PlasticSourceControl/Private/PlasticSourceControlUtils.cpp
@@ -476,8 +476,8 @@ static bool ParseWorkspaceInformation(const TArray<FString>& InInfoMessages, int
 {
 	bool bResult = true;
 
-	// Get workspace status, in the form "cs:41@rep:UE4PlasticPlugin@repserver:localhost:8087" (disabled by the "--nostatus" flag)
-	//                                or "cs:41@rep:UE4PlasticPlugin@repserver:SRombauts@cloud" (when connected directly to the cloud)
+	// Get workspace status, in the form "cs:41@rep:UEPlasticPlugin@repserver:localhost:8087" (disabled by the "--nostatus" flag)
+	//                                or "cs:41@rep:UEPlasticPlugin@repserver:SRombauts@cloud" (when connected directly to the cloud)
 	if (InInfoMessages.Num() > 0)
 	{
 		static const FString ChangesetPrefix(TEXT("cs:"));
@@ -498,7 +498,7 @@ static bool ParseWorkspaceInformation(const TArray<FString>& InInfoMessages, int
 			bResult = false;
 		}
 	}
-	// Get the branch name, in the form "Branch /main@UE4PlasticPluginDev" (enabled by the "--wkconfig" flag)
+	// Get the branch name, in the form "Branch /main@UEPlasticPluginDev" (enabled by the "--wkconfig" flag)
 	if (InInfoMessages.Num() > 1)
 	{
 		OutBranchName = InInfoMessages[1];
@@ -883,7 +883,7 @@ static bool RunStatus(const TArray<FString>& InFiles, const EConcurrency::Type I
 }
 
 // Parse the fileinfo output format "{RevisionChangeset};{RevisionHeadChangeset};{RepSpec};{LockedBy};{LockedWhere}"
-// for example "40;41;repo@server:port;srombauts;UE4PlasticPluginDev"
+// for example "40;41;repo@server:port;srombauts;UEPlasticPluginDev"
 class FPlasticFileinfoParser
 {
 public:
@@ -1057,9 +1057,9 @@ bool RunCheckMergeStatus(const TArray<FString>& InFiles, TArray<FString>& OutErr
 			// @todo: temporary debug logs
 			UE_LOG(LogSourceControl, Log, TEXT("RunCheckMergeStatus: %s:\n%s"), *MergeProgressFilename, *MergeProgressContent);
 			// Content is in one line, looking like the following:
-			// Target: mount:56e62dd7-241f-41e9-8c6b-dd4ca4513e62#/#UE4MergeTest@localhost:8087 merged from: Merge 4
-			// Target: mount:56e62dd7-241f-41e9-8c6b-dd4ca4513e62#/#UE4MergeTest@localhost:8087 merged from: Cherrypicking 3
-			// Target: mount:56e62dd7-241f-41e9-8c6b-dd4ca4513e62#/#UE4MergeTest@localhost:8087 merged from: IntervalCherrypick 2 4
+			// Target: mount:56e62dd7-241f-41e9-8c6b-dd4ca4513e62#/#UEMergeTest@localhost:8087 merged from: Merge 4
+			// Target: mount:56e62dd7-241f-41e9-8c6b-dd4ca4513e62#/#UEMergeTest@localhost:8087 merged from: Cherrypicking 3
+			// Target: mount:56e62dd7-241f-41e9-8c6b-dd4ca4513e62#/#UEMergeTest@localhost:8087 merged from: IntervalCherrypick 2 4
 			// 1) Extract the word after "merged from: "
 			static const FString MergeFromString(TEXT("merged from: "));
 			const int32 MergeFromIndex = MergeProgressContent.Find(MergeFromString, ESearchCase::CaseSensitive);
diff --git a/Engine/Source/Developer/AssetTools/Private/AssetViewUtils.cpp b/Engine/Source/Developer/AssetTools/Private/AssetViewUtils.cpp
--- a/Engine/Source/Developer/AssetTools/Private/AssetViewUtils.cpp
+++ b/Engine/Source/Developer/AssetTools/Private/AssetViewUtils.cpp
@@ -1381,7 +1381,7 @@ int32 AssetViewUtils::GetPackageLengthForCooking(const FString& PackageName, boo
 		if (IsInternalBuild)
 		{
 			// We assume a constant size for the build machine base path, so strip either the root or game path from the start
-			// (depending on whether the project is part of the main UE4 source tree or located elsewhere)
+			// (depending on whether the project is part of the main UE source tree or located elsewhere)
 			FString CookDirWithoutBasePath = AbsoluteCookPath;
 			if (CookDirWithoutBasePath.StartsWith(AbsoluteRootPath, ESearchCase::CaseSensitive))
 			{
@@ -1392,7 +1392,7 @@ int32 AssetViewUtils::GetPackageLengthForCooking(const FString& PackageName, boo
 				CookDirWithoutBasePath.RemoveFromStart(AbsoluteCookPath, ESearchCase::CaseSensitive);
 			}
 			
-			FString AbsoluteBuildMachineCookPathToAsset = FString(TEXT("D:/BuildFarm/buildmachine_++depot+UE4-Releases+4.10")) / CookDirWithoutBasePath / AssetPathWithinCookDir;
+			FString AbsoluteBuildMachineCookPathToAsset = FString(TEXT("D:/BuildFarm/buildmachine_++depot+UE-Releases+XX.XX")) / CookDirWithoutBasePath / AssetPathWithinCookDir;
 
 			// only add game name padding to project plugins so that they can ported to other projects
 			if (bIsPluginAsset && bIsProjectAsset)			{
@@ -1423,7 +1423,7 @@ bool AssetViewUtils::IsValidPackageForCooking(const FString& PackageName, FText&
 		// See TTP# 332328:
 		// The following checks are done mostly to prevent / alleviate the problems that "long" paths are causing with the BuildFarm and cooked builds.
 		// The BuildFarm uses a verbose path to encode extra information to provide more information when things fail, however this makes the path limitation a problem.
-		//	- We assume a base path of D:/BuildFarm/buildmachine_++depot+UE4-Releases+4.10/
+		//	- We assume a base path of length: D:/BuildFarm/buildmachine_++depot+UE-Releases+XX.XX/
 		//	- We assume the game name is 20 characters (the maximum allowed) to make sure that content can be ported between projects
 		//	- We calculate the cooked game path relative to the game root (eg, Showcases/Infiltrator/Saved/Cooked/Windows/Infiltrator)
 		//	- We calculate the asset path relative to (and including) the Content directory (eg, Content/Environment/Infil1/Infil1_Underground/Infrastructure/Model/SM_Infil1_Tunnel_Ceiling_Pipes_1xEntryCurveOuter_Double.uasset)
diff --git a/Engine/Source/Editor/Persona/Private/AnimationEditorViewportClient.cpp b/Engine/Source/Editor/Persona/Private/AnimationEditorViewportClient.cpp
--- a/Engine/Source/Editor/Persona/Private/AnimationEditorViewportClient.cpp
+++ b/Engine/Source/Editor/Persona/Private/AnimationEditorViewportClient.cpp
@@ -38,9 +38,7 @@
 #include "UnrealWidget.h"
 
 namespace {
-	// Value from UE3
 	static const float AnimationEditorViewport_RotateSpeed = 0.02f;
-	// Value from UE3
 	static const float AnimationEditorViewport_TranslateSpeed = 0.25f;
 	// follow camera feature
 	static const float FollowCamera_InterpSpeed = 4.f;
