PR #8913: Don't report warning when a file is locked in Plastic SCM (Contributed by [REDACTED_PERSON])

diff --git a/Engine/Plugins/Developer/PlasticSourceControl/Source/PlasticSourceControl/Private/PlasticSourceControlUtils.cpp b/Engine/Plugins/Developer/PlasticSourceControl/Source/PlasticSourceControl/Private/PlasticSourceControlUtils.cpp
--- a/Engine/Plugins/Developer/PlasticSourceControl/Source/PlasticSourceControl/Private/PlasticSourceControlUtils.cpp
+++ b/Engine/Plugins/Developer/PlasticSourceControl/Source/PlasticSourceControl/Private/PlasticSourceControlUtils.cpp
@@ -947,7 +947,7 @@ static void ParseFileinfoResults(const TArray<FString>& InResults, TArray<FPlast
 		if ((FileState.WorkspaceState != EWorkspaceState::CheckedOut) && (FileState.WorkspaceState != EWorkspaceState::Moved) && (0 < FileState.LockedBy.Len())) 
 		{
 			// @todo: temporary debug log
-			UE_LOG(LogSourceControl, Warning, TEXT("LockedByOther(%s) by '%s!=%s' (or %s!=%s)"), *File, *FileState.LockedBy, *Provider.GetUserName(), *FileState.LockedWhere, *Provider.GetWorkspaceName());
+			UE_LOG(LogSourceControl, Log, TEXT("LockedByOther(%s) by '%s!=%s' (or %s!=%s)"), *File, *FileState.LockedBy, *Provider.GetUserName(), *FileState.LockedWhere, *Provider.GetWorkspaceName());
 			FileState.WorkspaceState = EWorkspaceState::LockedByOther;
 		}
 
