Updated PlasticSCM from 1.4.6 to 1.5.0

#9156: Update PlasticSCM Plugin from 1.4.6 to 1.5.0 for UE5 (Contributed by [REDACTED_PERSON])

diff --git a/Engine/Plugins/Developer/PlasticSourceControl/PlasticSourceControl.uplugin b/Engine/Plugins/Developer/PlasticSourceControl/PlasticSourceControl.uplugin
--- a/Engine/Plugins/Developer/PlasticSourceControl/PlasticSourceControl.uplugin
+++ b/Engine/Plugins/Developer/PlasticSourceControl/PlasticSourceControl.uplugin
@@ -1,24 +1,27 @@
-{
-	"FileVersion": 3,
-	"Version": 40,
-	"VersionName": "1.4.6",
-	"FriendlyName": "Plastic SCM",
-	"Description": "Plastic source control management",
-	"Category": "Source Control",
-	"CreatedBy": "SRombauts for Codice Software.",
-	"CreatedByURL": "http://srombauts.github.com",
-	"DocsURL": "https://www.plasticscm.com/documentation.html",
-	"MarketplaceURL": "",
-	"SupportURL": "https://www.plasticscm.com/support.html",
-	"EnabledByDefault": true,
-	"CanContainContent": false,
-	"IsBetaVersion": false,
-	"Installed": false,
-	"Modules": [
-		{
-			"Name": "PlasticSourceControl",
-			"Type": "Editor",
-			"LoadingPhase": "Default"
-		}
-	]
+{
+	"FileVersion": 3,
+	"Version": 50,
+	"VersionName": "1.5.0",
+	"FriendlyName": "Plastic SCM",
+	"Description": "Plastic source control management",
+	"Category": "Source Control",
+	"CreatedBy": "SRombauts for Codice Software.",
+	"CreatedByURL": "http://srombauts.github.com",
+	"DocsURL": "[REDACTED_PROFILE_URL]/UE4PlasticPlugin",
+	"MarketplaceURL": "",
+	"SupportURL": "https://www.plasticscm.com/support",
+	"EnabledByDefault": true,
+	"CanContainContent": false,
+	"IsBetaVersion": false,
+	"Installed": false,
+	"Modules": [
+		{
+			"Name": "PlasticSourceControl",
+			"Type": "Editor",
+			"LoadingPhase": "Default",
+			"PlatformAllowList": [
+				"Win64"
+			]
+		}
+	]
 }
\ No newline at end of file
diff --git a/Engine/Plugins/Developer/PlasticSourceControl/Source/PlasticSourceControl/PlasticSourceControl.Build.cs b/Engine/Plugins/Developer/PlasticSourceControl/Source/PlasticSourceControl/PlasticSourceControl.Build.cs
--- a/Engine/Plugins/Developer/PlasticSourceControl/Source/PlasticSourceControl/PlasticSourceControl.Build.cs
+++ b/Engine/Plugins/Developer/PlasticSourceControl/Source/PlasticSourceControl/PlasticSourceControl.Build.cs
@@ -14,13 +14,13 @@ public class PlasticSourceControl : ModuleRules
 				"Slate",
 				"SlateCore",
 				"InputCore",
-				"EditorFramework",
 				"UnrealEd",
 				"LevelEditor",
 				"SourceControl",
 				"XmlParser",
 				"Projects",
 				"AssetRegistry",
+				"ToolMenus",
 			}
 		);
 	}
diff --git a/Engine/Plugins/Developer/PlasticSourceControl/Source/PlasticSourceControl/Private/PlasticSourceControlCommand.cpp b/Engine/Plugins/Developer/PlasticSourceControl/Source/PlasticSourceControl/Private/PlasticSourceControlCommand.cpp
--- a/Engine/Plugins/Developer/PlasticSourceControl/Source/PlasticSourceControl/Private/PlasticSourceControlCommand.cpp
+++ b/Engine/Plugins/Developer/PlasticSourceControl/Source/PlasticSourceControl/Private/PlasticSourceControlCommand.cpp
@@ -3,11 +3,11 @@
 #include "PlasticSourceControlCommand.h"
 #include "PlasticSourceControlModule.h"
 #include "ISourceControlOperation.h"
-#include "IPlasticSourceControlWorker.h"
 #include "Modules/ModuleManager.h"
+#include "HAL/PlatformTime.h"
 
 
-FPlasticSourceControlCommand::FPlasticSourceControlCommand(const TSharedRef<class ISourceControlOperation, ESPMode::ThreadSafe>& InOperation, const TSharedRef<class IPlasticSourceControlWorker, ESPMode::ThreadSafe>& InWorker, const FSourceControlOperationComplete& InOperationCompleteDelegate)
+FPlasticSourceControlCommand::FPlasticSourceControlCommand(const FSourceControlOperationRef& InOperation, const FPlasticSourceControlWorkerRef& InWorker, const FSourceControlOperationComplete& InOperationCompleteDelegate)
 	: Operation(InOperation)
 	, Worker(InWorker)
 	, OperationCompleteDelegate(InOperationCompleteDelegate)
@@ -16,6 +16,7 @@ FPlasticSourceControlCommand::FPlasticSourceControlCommand(const TSharedRef<clas
 	, bConnectionDropped(false)
 	, bAutoDelete(true)
 	, Concurrency(EConcurrency::Synchronous)
+	, StartTimestamp(FPlatformTime::Seconds())
 {
 	// grab the providers settings here, so we don't access them once the worker thread is launched
 	check(IsInGameThread());
diff --git a/Engine/Plugins/Developer/PlasticSourceControl/Source/PlasticSourceControl/Private/PlasticSourceControlCommand.h b/Engine/Plugins/Developer/PlasticSourceControl/Source/PlasticSourceControl/Private/PlasticSourceControlCommand.h
--- a/Engine/Plugins/Developer/PlasticSourceControl/Source/PlasticSourceControl/Private/PlasticSourceControlCommand.h
+++ b/Engine/Plugins/Developer/PlasticSourceControl/Source/PlasticSourceControl/Private/PlasticSourceControlCommand.h
@@ -5,6 +5,7 @@
 #include "CoreMinimal.h"
 #include "Misc/IQueuedWork.h"
 #include "ISourceControlProvider.h"
+#include "IPlasticSourceControlWorker.h"
 
 /**
  * Used to execute Plastic commands multi-threaded.
@@ -13,7 +14,7 @@ class FPlasticSourceControlCommand : public IQueuedWork
 {
 public:
 
-	FPlasticSourceControlCommand(const TSharedRef<class ISourceControlOperation, ESPMode::ThreadSafe>& InOperation, const TSharedRef<class IPlasticSourceControlWorker, ESPMode::ThreadSafe>& InWorker, const FSourceControlOperationComplete& InOperationCompleteDelegate = FSourceControlOperationComplete() );
+	FPlasticSourceControlCommand(const FSourceControlOperationRef& InOperation, const FPlasticSourceControlWorkerRef& InWorker, const FSourceControlOperationComplete& InOperationCompleteDelegate = FSourceControlOperationComplete() );
 
 	/**
 	 * This is where the real thread work is done. All work that is done for
@@ -40,10 +41,10 @@ public:
 	FString PathToWorkspaceRoot;
 
 	/** Operation we want to perform - contains outward-facing parameters & results */
-	TSharedRef<class ISourceControlOperation, ESPMode::ThreadSafe> Operation;
+	FSourceControlOperationRef Operation;
 
 	/** The object that will actually do the work */
-	TSharedRef<class IPlasticSourceControlWorker, ESPMode::ThreadSafe> Worker;
+	FPlasticSourceControlWorkerRef Worker;
 
 	/** Delegate to notify when this operation completes */
 	FSourceControlOperationComplete OperationCompleteDelegate;
@@ -78,6 +79,9 @@ public:
 	/** Whether we are running multi-treaded in the background, or blocking the main thread */
 	EConcurrency::Type Concurrency;
 
+	/** Timestamp of when the command was issued */
+	const double StartTimestamp;
+
 	/** Files to perform this operation on */
 	TArray< FString > Files;
 
diff --git a/Engine/Plugins/Developer/PlasticSourceControl/Source/PlasticSourceControl/Private/PlasticSourceControlConsole.cpp b/Engine/Plugins/Developer/PlasticSourceControl/Source/PlasticSourceControl/Private/PlasticSourceControlConsole.cpp
new file mode 100644
--- /dev/null
+++ b/Engine/Plugins/Developer/PlasticSourceControl/Source/PlasticSourceControl/Private/PlasticSourceControlConsole.cpp
@@ -0,0 +1,42 @@
+// Copyright Epic Games, Inc. All Rights Reserved.
+
+#include "PlasticSourceControlConsole.h"
+
+#include "HAL/IConsoleManager.h"
+#include "ISourceControlModule.h"
+
+#include "PlasticSourceControlUtils.h"
+
+void FPlasticSourceControlConsole::Register()
+{
+	CmConsoleCommand = MakeUnique<FAutoConsoleCommand>(
+		TEXT("cm"),
+		TEXT("PlasticSCM Command Line Interface.\n")
+		TEXT("Run any 'cm' command directly from the Unreal Editor Console.\n")
+		TEXT("Type 'cm showcommands' to get a command list."),
+		FConsoleCommandWithArgsDelegate::CreateRaw(this, &FPlasticSourceControlConsole::ExecutePlasticConsoleCommand)
+	);
+}
+
+void FPlasticSourceControlConsole::Unregister()
+{
+	CmConsoleCommand.Reset();
+}
+
+void FPlasticSourceControlConsole::ExecutePlasticConsoleCommand(const TArray<FString>& a_args)
+{
+	// If called with no argument, explicitely call "cm help" instead to mimic the cm CLI behavior.
+	if (a_args.Num() < 1)
+	{
+		ExecutePlasticConsoleCommand(TArray<FString>({TEXT("help")}));
+		return;
+	}
+
+	FString Results;
+	FString Errors;
+	const FString Command = a_args[0];
+	TArray<FString> Parameters = a_args;
+	Parameters.RemoveAt(0);
+	PlasticSourceControlUtils::RunCommandInternal(Command, Parameters, TArray<FString>(), EConcurrency::Synchronous, Results, Errors);
+	UE_LOG(LogSourceControl, Log, TEXT("Output:\n%s"), *Results);
+}
diff --git a/Engine/Plugins/Developer/PlasticSourceControl/Source/PlasticSourceControl/Private/PlasticSourceControlConsole.h b/Engine/Plugins/Developer/PlasticSourceControl/Source/PlasticSourceControl/Private/PlasticSourceControlConsole.h
new file mode 100644
--- /dev/null
+++ b/Engine/Plugins/Developer/PlasticSourceControl/Source/PlasticSourceControl/Private/PlasticSourceControlConsole.h
@@ -0,0 +1,26 @@
+// Copyright Epic Games, Inc. All Rights Reserved.
+
+#pragma once
+
+#include "CoreMinimal.h"
+#include "HAL/IConsoleManager.h"
+
+/**
+ * Editor only console commands.
+ *
+ * Such commands can be executed from the editor output log window, but also from command line arguments,
+ * from Editor Blueprints utilities, or from C++ Code using. eg. GEngine->Exec("cm status");
+ */
+class FPlasticSourceControlConsole
+{
+public:
+	void Register();
+	void Unregister();
+
+private:
+	// PlasticSCM Command Line Interface: Run 'cm' commands directly from the Unreal Editor Console.
+	void ExecutePlasticConsoleCommand(const TArray<FString>& a_args);
+
+	/** Console command for interacting with 'cm' CLI directly */
+	TUniquePtr<FAutoConsoleCommand> CmConsoleCommand;
+};
diff --git a/Engine/Plugins/Developer/PlasticSourceControl/Source/PlasticSourceControl/Private/PlasticSourceControlMenu.cpp b/Engine/Plugins/Developer/PlasticSourceControl/Source/PlasticSourceControl/Private/PlasticSourceControlMenu.cpp
--- a/Engine/Plugins/Developer/PlasticSourceControl/Source/PlasticSourceControl/Private/PlasticSourceControlMenu.cpp
+++ b/Engine/Plugins/Developer/PlasticSourceControl/Source/PlasticSourceControl/Private/PlasticSourceControlMenu.cpp
@@ -10,6 +10,7 @@
 #include "ISourceControlOperation.h"
 #include "SourceControlOperations.h"
 
+#include "Interfaces/IPluginManager.h"
 #include "LevelEditor.h"
 #include "Widgets/Notifications/SNotificationList.h"
 #include "Framework/Notifications/NotificationManager.h"
@@ -22,30 +23,31 @@
 
 #include "Logging/MessageLog.h"
 
-static const FName PlasticSourceControlMenuTabName("PlasticSourceControlMenu");
+#include "ToolMenus.h"
+#include "ToolMenuContext.h"
+#include "ToolMenuMisc.h"
 
 #define LOCTEXT_NAMESPACE "PlasticSourceControl"
 
 void FPlasticSourceControlMenu::Register()
 {
-	// Register the extension with the level editor
-	FLevelEditorModule* LevelEditorModule = FModuleManager::GetModulePtr<FLevelEditorModule>("LevelEditor");
-	if (LevelEditorModule)
+	// Register the menu extension with the level editor
+	FToolMenuOwnerScoped SourceControlMenuOwner("PlasticSourceControlMenu");
+	if (UToolMenus* ToolMenus = UToolMenus::Get())
 	{
-		FLevelEditorModule::FLevelEditorMenuExtender ViewMenuExtender = FLevelEditorModule::FLevelEditorMenuExtender::CreateRaw(this, &FPlasticSourceControlMenu::OnExtendLevelEditorViewMenu);
-		auto& MenuExtenders = LevelEditorModule->GetAllLevelEditorToolbarSourceControlMenuExtenders();
-		MenuExtenders.Add(ViewMenuExtender);
-		ViewMenuExtenderHandle = MenuExtenders.Last().GetHandle();
+		UToolMenu* SourceControlMenu = ToolMenus->ExtendMenu("StatusBar.ToolBar.SourceControl");
+		FToolMenuSection& Section = SourceControlMenu->AddSection("PlasticSourceControlActions", LOCTEXT("PlasticSourceControlMenuHeadingActions", "Plastic SCM"), FToolMenuInsert(NAME_None, EToolMenuInsertType::First));
+
+		AddMenuExtension(Section);
 	}
 }
 
 void FPlasticSourceControlMenu::Unregister()
 {
-	// Unregister the level editor extensions
-	FLevelEditorModule* LevelEditorModule = FModuleManager::GetModulePtr<FLevelEditorModule>("LevelEditor");
-	if (LevelEditorModule)
+	// Unregister the menu extension from the level editor
+	if (UToolMenus* ToolMenus = UToolMenus::Get())
 	{
-		LevelEditorModule->GetAllLevelEditorToolbarSourceControlMenuExtenders().RemoveAll([=](const FLevelEditorModule::FLevelEditorMenuExtender& Extender) { return Extender.GetHandle() == ViewMenuExtenderHandle; });
+		UToolMenus::Get()->UnregisterOwnerByName("PlasticSourceControlMenu");
 	}
 }
 
@@ -292,7 +294,8 @@ void FPlasticSourceControlMenu::RefreshClicked()
 		FPlasticSourceControlModule& PlasticSourceControl = FModuleManager::LoadModuleChecked<FPlasticSourceControlModule>("PlasticSourceControl");
 		FPlasticSourceControlProvider& Provider = PlasticSourceControl.GetProvider();
 		TSharedRef<FUpdateStatus, ESPMode::ThreadSafe> RefreshOperation = ISourceControlOperation::Create<FUpdateStatus>();
-		RefreshOperation->SetCheckingAllFiles(true);
+		// This is the flag used by the Content Browser's "Checkout" filter to trigger a full status update
+		RefreshOperation->SetGetOpenedOnly(true);
 		const ECommandResult::Type Result = Provider.Execute(RefreshOperation, TArray<FString>(), EConcurrency::Asynchronous, FSourceControlOperationComplete::CreateRaw(this, &FPlasticSourceControlMenu::OnSourceControlOperationComplete));
 		if (Result == ECommandResult::Succeeded)
 		{
@@ -313,6 +316,26 @@ void FPlasticSourceControlMenu::RefreshClicked()
 	}
 }
 
+void FPlasticSourceControlMenu::VisitDocsURLClicked()
+{
+	// Grab the URL from the uplugin file
+	const TSharedPtr<IPlugin> Plugin = IPluginManager::Get().FindPlugin(TEXT("PlasticSourceControl"));
+	if (Plugin.IsValid())
+	{
+		FPlatformProcess::LaunchURL(*Plugin->GetDescriptor().DocsURL, NULL, NULL);
+	}
+}
+
+void FPlasticSourceControlMenu::VisitSupportURLClicked()
+{
+	// Grab the URL from the uplugin file
+	const TSharedPtr<IPlugin> Plugin = IPluginManager::Get().FindPlugin(TEXT("PlasticSourceControl"));
+	if (Plugin.IsValid())
+	{
+		FPlatformProcess::LaunchURL(*Plugin->GetDescriptor().SupportURL, NULL, NULL);
+	}
+}
+
 // Display an ongoing notification during the whole operation
 void FPlasticSourceControlMenu::DisplayInProgressNotification(const FText& InOperationInProgressString)
 {
@@ -388,9 +411,10 @@ void FPlasticSourceControlMenu::OnSourceControlOperationComplete(const FSourceCo
 	}
 }
 
-void FPlasticSourceControlMenu::AddMenuExtension(FMenuBuilder& Builder)
+void FPlasticSourceControlMenu::AddMenuExtension(FToolMenuSection& Menu)
 {
-	Builder.AddMenuEntry(
+	Menu.AddMenuEntry(
+		"PlasticSync",
 		LOCTEXT("PlasticSync",			"Sync/Update Workspace"),
 		LOCTEXT("PlasticSyncTooltip",	"Update all files in the workspace to the latest version."),
 		FSlateIcon(FAppStyle::GetAppStyleSetName(), "SourceControl.Actions.Sync"),
@@ -400,7 +424,8 @@ void FPlasticSourceControlMenu::AddMenuExtension(FMenuBuilder& Builder)
 		)
 	);
 
-	Builder.AddMenuEntry(
+	Menu.AddMenuEntry(
+		"PlasticRevertUnchanged",
 		LOCTEXT("PlasticRevertUnchanged",			"Revert Unchanged"),
 		LOCTEXT("PlasticRevertUnchangedTooltip",	"Revert checked-out but unchanged files in the workspace."),
 		FSlateIcon(FAppStyle::GetAppStyleSetName(), "SourceControl.Actions.Revert"),
@@ -410,7 +435,8 @@ void FPlasticSourceControlMenu::AddMenuExtension(FMenuBuilder& Builder)
 		)
 	);
 
-	Builder.AddMenuEntry(
+	Menu.AddMenuEntry(
+		"PlasticRevertAll",
 		LOCTEXT("PlasticRevertAll",			"Revert All"),
 		LOCTEXT("PlasticRevertAllTooltip",	"Revert all files in the workspace to their controlled/unchanged state."),
 		FSlateIcon(FAppStyle::GetAppStyleSetName(), "SourceControl.Actions.Revert"),
@@ -420,7 +446,8 @@ void FPlasticSourceControlMenu::AddMenuExtension(FMenuBuilder& Builder)
 		)
 	);
 
-	Builder.AddMenuEntry(
+	Menu.AddMenuEntry(
+		"PlasticRefresh",
 		LOCTEXT("PlasticRefresh",			"Refresh"),
 		LOCTEXT("PlasticRefreshTooltip",	"Update the source control status of all files in the workspace."),
 		FSlateIcon(FAppStyle::GetAppStyleSetName(), "SourceControl.Actions.Refresh"),
@@ -429,19 +456,29 @@ void FPlasticSourceControlMenu::AddMenuExtension(FMenuBuilder& Builder)
 			FCanExecuteAction()
 		)
 	);
-}
 
-TSharedRef<FExtender> FPlasticSourceControlMenu::OnExtendLevelEditorViewMenu(const TSharedRef<FUICommandList> CommandList)
-{
-	TSharedRef<FExtender> Extender(new FExtender());
-
-	Extender->AddMenuExtension(
-		"SourceControlActions",
-		EExtensionHook::After,
-		nullptr,
-		FMenuExtensionDelegate::CreateRaw(this, &FPlasticSourceControlMenu::AddMenuExtension));
+	Menu.AddMenuEntry(
+		"PlasticDocsURL",
+		LOCTEXT("PlasticDocsURL",			"Plugin's Documentation"),
+		LOCTEXT("PlasticDocsURLTooltip",	"Visit documentation of the plugin on Github."),
+		FSlateIcon(FAppStyle::GetAppStyleSetName(), "Icons.Documentation"),
+		FUIAction(
+			FExecuteAction::CreateRaw(this, &FPlasticSourceControlMenu::VisitDocsURLClicked),
+			FCanExecuteAction()
+		)
+	);
 
-	return Extender;
+	Menu.AddMenuEntry(
+		"PlasticSupportURL",
+		LOCTEXT("PlasticSupportURL",		"Plastic SCM Support"),
+		LOCTEXT("PlasticSupportURLTooltip",	"Visit official support for Plastic SCM."),
+		FSlateIcon(FAppStyle::GetAppStyleSetName(), "Icons.Support"),
+		FUIAction(
+			FExecuteAction::CreateRaw(this, &FPlasticSourceControlMenu::VisitSupportURLClicked),
+			FCanExecuteAction()
+		)
+	);
 }
 
+
 #undef LOCTEXT_NAMESPACE
diff --git a/Engine/Plugins/Developer/PlasticSourceControl/Source/PlasticSourceControl/Private/PlasticSourceControlMenu.h b/Engine/Plugins/Developer/PlasticSourceControl/Source/PlasticSourceControl/Private/PlasticSourceControlMenu.h
--- a/Engine/Plugins/Developer/PlasticSourceControl/Source/PlasticSourceControl/Private/PlasticSourceControlMenu.h
+++ b/Engine/Plugins/Developer/PlasticSourceControl/Source/PlasticSourceControl/Private/PlasticSourceControlMenu.h
@@ -5,8 +5,7 @@
 #include "CoreMinimal.h"
 #include "ISourceControlProvider.h"
 
-class FToolBarBuilder;
-class FMenuBuilder;
+struct FToolMenuSection;
 
 /** Plastic SCM extension of the Source Control toolbar menu */
 class FPlasticSourceControlMenu
@@ -20,6 +19,8 @@ public:
 	void RevertUnchangedClicked();
 	void RevertAllClicked();
 	void RefreshClicked();
+	void VisitDocsURLClicked();
+	void VisitSupportURLClicked();
 
 private:
 	bool IsSourceControlConnected() const;
@@ -29,7 +30,7 @@ private:
 	TArray<UPackage*>	UnlinkPackages(const TArray<FString>& InPackageNames);
 	void				ReloadPackages(TArray<UPackage*>& InPackagesToReload);
 
-	void AddMenuExtension(FMenuBuilder& Builder);
+	void AddMenuExtension(FToolMenuSection& Builder);
 
 	TSharedRef<class FExtender> OnExtendLevelEditorViewMenu(const TSharedRef<class FUICommandList> CommandList);
 
@@ -39,7 +40,6 @@ private:
 	void DisplayFailureNotification(const FName& InOperationName);
 
 private:
-	FDelegateHandle ViewMenuExtenderHandle;
 
 	/** Loaded packages to reload after a Sync or Revert operation */
 	TArray<UPackage*> PackagesToReload;
diff --git a/Engine/Plugins/Developer/PlasticSourceControl/Source/PlasticSourceControl/Private/PlasticSourceControlModule.cpp b/Engine/Plugins/Developer/PlasticSourceControl/Source/PlasticSourceControl/Private/PlasticSourceControlModule.cpp
--- a/Engine/Plugins/Developer/PlasticSourceControl/Source/PlasticSourceControl/Private/PlasticSourceControlModule.cpp
+++ b/Engine/Plugins/Developer/PlasticSourceControl/Source/PlasticSourceControl/Private/PlasticSourceControlModule.cpp
@@ -8,7 +8,7 @@
 #define LOCTEXT_NAMESPACE "PlasticSourceControl"
 
 template<typename Type>
-static TSharedRef<IPlasticSourceControlWorker, ESPMode::ThreadSafe> CreateWorker()
+static FPlasticSourceControlWorkerRef CreateWorker()
 {
 	return MakeShareable( new Type() );
 }
diff --git a/Engine/Plugins/Developer/PlasticSourceControl/Source/PlasticSourceControl/Private/PlasticSourceControlModule.h b/Engine/Plugins/Developer/PlasticSourceControl/Source/PlasticSourceControl/Private/PlasticSourceControlModule.h
--- a/Engine/Plugins/Developer/PlasticSourceControl/Source/PlasticSourceControl/Private/PlasticSourceControlModule.h
+++ b/Engine/Plugins/Developer/PlasticSourceControl/Source/PlasticSourceControl/Private/PlasticSourceControlModule.h
@@ -9,7 +9,7 @@
 #include "PlasticSourceControlProvider.h"
 
 /**
- * UEPlasticPlugin is a simple Plastic SCM Source Control Plugin for Unreal Engine
+ * PlasticSourceControl is a simple Plastic SCM Source Control Plugin for Unreal Engine
  *
  * Written and contributed by [REDACTED_PERSON] ([REDACTED_EMAIL])
  */
diff --git a/Engine/Plugins/Developer/PlasticSourceControl/Source/PlasticSourceControl/Private/PlasticSourceControlOperations.cpp b/Engine/Plugins/Developer/PlasticSourceControl/Source/PlasticSourceControl/Private/PlasticSourceControlOperations.cpp
--- a/Engine/Plugins/Developer/PlasticSourceControl/Source/PlasticSourceControl/Private/PlasticSourceControlOperations.cpp
+++ b/Engine/Plugins/Developer/PlasticSourceControl/Source/PlasticSourceControl/Private/PlasticSourceControlOperations.cpp
@@ -46,6 +46,16 @@ FText FPlasticMakeWorkspace::GetInProgressString() const
 	return LOCTEXT("SourceControl_MakeWorkspace", "Create a new Repository and initialize the Workspace");
 }
 
+static bool AreAllFiles(const TArray<FString>& InFiles)
+{
+	for (const FString& File : InFiles)
+	{
+		if (File.IsEmpty() || File[File.Len() - 1] == TEXT('/'))
+			return false;
+	}
+	return true;
+}
+
 
 FName FPlasticConnectWorker::GetName() const
 {
@@ -72,15 +82,14 @@ bool FPlasticConnectWorker::Execute(FPlasticSourceControlCommand& InCommand)
 				InCommand.bCommandSuccessful = PlasticSourceControlUtils::RunCommand(TEXT("checkconnection"), TArray<FString>(), TArray<FString>(), InCommand.Concurrency, InCommand.InfoMessages, InCommand.ErrorMessages);
 				if (InCommand.bCommandSuccessful)
 				{
-					// Then, update the status of assets in Content/ directory and also Config files,
+					// Now update the status of assets in the Content directory
 					// but only on real (re-)connection (but not each time Login() is called by Rename or Fixup Redirector command to check connection)
 					// and only if enabled in the settings
-					if (!PlasticSourceControl.GetProvider().IsAvailable() && PlasticSourceControl.AccessSettings().UpdateStatusAtStartup())
+					if (!PlasticSourceControl.GetProvider().IsAvailable() && PlasticSourceControl.AccessSettings().GetUpdateStatusAtStartup())
 					{
-						TArray<FString> ProjectDirs;
-						ProjectDirs.Add(FPaths::ConvertRelativePathToFull(FPaths::ProjectContentDir()));
-						ProjectDirs.Add(FPaths::ConvertRelativePathToFull(FPaths::ProjectConfigDir()));
-						PlasticSourceControlUtils::RunUpdateStatus(ProjectDirs, false, InCommand.Concurrency, InCommand.ErrorMessages, States, InCommand.ChangesetNumber, InCommand.BranchName);
+						TArray<FString> ContentDir;
+						ContentDir.Add(FPaths::ConvertRelativePathToFull(FPaths::ProjectContentDir()));
+						PlasticSourceControlUtils::RunUpdateStatus(ContentDir, false, InCommand.Concurrency, InCommand.ErrorMessages, States, InCommand.ChangesetNumber, InCommand.BranchName);
 					}
 				}
 				else
@@ -174,7 +183,7 @@ bool FPlasticCheckInWorker::Execute(FPlasticSourceControlCommand& InCommand)
 
 	// make a temp file to place our commit message in
 	FScopedTempFile CommitMsgFile(Operation->GetDescription());
-	if (CommitMsgFile.GetFilename().Len() > 0)
+	if (!CommitMsgFile.GetFilename().IsEmpty())
 	{
 		TArray<FString> Parameters;
 		FString ParamCommitMsgFilename = TEXT("--commentsfile=\"");
@@ -236,11 +245,24 @@ bool FPlasticMarkForAddWorker::Execute(FPlasticSourceControlCommand& InCommand)
 	check(InCommand.Operation->GetName() == GetName());
 
 	TArray<FString> Parameters;
-	Parameters.Add(TEXT("--parents"));
+	Parameters.Add(TEXT("--parents")); // NOTE: deprecated in 8.0.16.3100 when it became the default https://www.plasticscm.com/download/releasenotes/8.0.16.3100
+	// Note: using "?" is a workaround to trigger the Plastic's "SkipIgnored" internal flag meaning "don't add file that are ignored":
+	//			options.SkipIgnored = cla.GetWildCardArguments().Count > 0;
+	//		 It's behavior is similar as Subversion:
+	//  		if you explicitely add one file that is ignored, "cm" will happily accept it and add it,
+	//			if you try to add a set of files with a pattern, "cm" will skip the files that are ignored and only add the other ones
+	// TODO: provide an updated version of "cm" with a new flag like --applyignorerules
+	if (AreAllFiles(InCommand.Files))
+	{
+		Parameters.Add(TEXT("?"));	// needed only when used with a list of files
+	}
+	else
+	{
+		Parameters.Add(TEXT("-R"));	// needed only at the time of workspace creation, to add directories recursively
+	}
 	// Detect special case for a partial checkout (CS:-1 in Gluon mode)!
 	if (-1 != InCommand.ChangesetNumber)
 	{
-		Parameters.Add(TEXT("-R")); // needed at the time of workspace creation, but not working in a partial workspace
 		InCommand.bCommandSuccessful = PlasticSourceControlUtils::RunCommand(TEXT("add"), Parameters, InCommand.Files, InCommand.Concurrency, InCommand.InfoMessages, InCommand.ErrorMessages);
 	}
 	else
@@ -303,7 +325,7 @@ bool FPlasticRevertWorker::Execute(FPlasticSourceControlCommand& InCommand)
 
 	InCommand.bCommandSuccessful = true;
 
-	for(const FString& File : InCommand.Files)
+	for (const FString& File : InCommand.Files)
 	{
 		TSharedRef<FPlasticSourceControlState, ESPMode::ThreadSafe> State = Provider.GetStateInternal(File);
 
@@ -364,11 +386,10 @@ bool FPlasticRevertUnchangedWorker::Execute(FPlasticSourceControlCommand& InComm
 	// revert the checkout of all unchanged files recursively
 	InCommand.bCommandSuccessful = PlasticSourceControlUtils::RunCommand(TEXT("uncounchanged"), Parameters, InCommand.Files, InCommand.Concurrency, InCommand.InfoMessages, InCommand.ErrorMessages);
 
-	// Now update the status of assets in Content/ directory and also Config files
-	TArray<FString> ProjectDirs;
-	ProjectDirs.Add(FPaths::ConvertRelativePathToFull(FPaths::ProjectContentDir()));
-	ProjectDirs.Add(FPaths::ConvertRelativePathToFull(FPaths::ProjectConfigDir()));
-	PlasticSourceControlUtils::RunUpdateStatus(ProjectDirs, false, InCommand.Concurrency, InCommand.ErrorMessages, States, InCommand.ChangesetNumber, InCommand.BranchName);
+	// Now update the status of assets in the Content directory
+	TArray<FString> ContentDir;
+	ContentDir.Add(FPaths::ConvertRelativePathToFull(FPaths::ProjectContentDir()));
+	PlasticSourceControlUtils::RunUpdateStatus(ContentDir, false, InCommand.Concurrency, InCommand.ErrorMessages, States, InCommand.ChangesetNumber, InCommand.BranchName);
 
 	return InCommand.bCommandSuccessful;
 }
@@ -400,11 +421,10 @@ bool FPlasticRevertAllWorker::Execute(FPlasticSourceControlCommand& InCommand)
 		InCommand.bCommandSuccessful = PlasticSourceControlUtils::RunCommand(TEXT("partial undocheckout"), Parameters, InCommand.Files, InCommand.Concurrency, InCommand.InfoMessages, InCommand.ErrorMessages);
 	}
 
-	// Now update the status of assets in Content/ directory and also Config files
-	TArray<FString> ProjectDirs;
-	ProjectDirs.Add(FPaths::ConvertRelativePathToFull(FPaths::ProjectContentDir()));
-	ProjectDirs.Add(FPaths::ConvertRelativePathToFull(FPaths::ProjectConfigDir()));
-	PlasticSourceControlUtils::RunUpdateStatus(ProjectDirs, false, InCommand.Concurrency, InCommand.ErrorMessages, States, InCommand.ChangesetNumber, InCommand.BranchName);
+	// Now update the status of assets in the Content directory
+	TArray<FString> ContentDir;
+	ContentDir.Add(FPaths::ConvertRelativePathToFull(FPaths::ProjectContentDir()));
+	PlasticSourceControlUtils::RunUpdateStatus(ContentDir, false, InCommand.Concurrency, InCommand.ErrorMessages, States, InCommand.ChangesetNumber, InCommand.BranchName);
 
 	return InCommand.bCommandSuccessful;
 }
@@ -456,9 +476,9 @@ bool FPlasticUpdateStatusWorker::Execute(FPlasticSourceControlCommand& InCommand
 	check(InCommand.Operation->GetName() == GetName());
 	TSharedRef<FUpdateStatus, ESPMode::ThreadSafe> Operation = StaticCastSharedRef<FUpdateStatus>(InCommand.Operation);
 
-	// @todo: temporary debug log
-	UE_LOG(LogSourceControl, Log, TEXT("status (of %d files, ShouldCheckAllFiles=%d, ShouldUpdateHistory=%d, ShouldGetOpenedOnly=%d, ShouldUpdateModifiedState=%d)"),
-		InCommand.Files.Num(), Operation->ShouldCheckAllFiles()?1:0, Operation->ShouldUpdateHistory()?1:0, Operation->ShouldGetOpenedOnly()?1:0, Operation->ShouldUpdateModifiedState()?1:0);
+	// Note: ShouldCheckAllFiles is never set to true (SetCheckingAllFiles)
+	UE_LOG(LogSourceControl, Log, TEXT("status (of %d files, ShouldUpdateHistory=%d, ShouldGetOpenedOnly=%d, ShouldUpdateModifiedState=%d)"),
+		InCommand.Files.Num(), Operation->ShouldUpdateHistory(), Operation->ShouldGetOpenedOnly(), Operation->ShouldUpdateModifiedState());
 
 	if (InCommand.Files.Num() > 0)
 	{
@@ -470,54 +490,62 @@ bool FPlasticUpdateStatusWorker::Execute(FPlasticSourceControlCommand& InCommand
 			UE_LOG(LogSourceControl, Error, TEXT("FPlasticUpdateStatusWorker(ErrorMessages.Num()=%d) => checkconnection"), InCommand.ErrorMessages.Num());
 			// In case of error, execute a 'checkconnection' command to check the connectivity of the server.
 			InCommand.bConnectionDropped = !PlasticSourceControlUtils::RunCommand(TEXT("checkconnection"), TArray<FString>(), TArray<FString>(), InCommand.Concurrency, InCommand.InfoMessages, InCommand.ErrorMessages);
+			return false;
 		}
-		else
+
+		if (Operation->ShouldUpdateHistory())
 		{
-			if (Operation->ShouldUpdateHistory())
+			// Get the history of the files (on all branches)
+			InCommand.bCommandSuccessful &= PlasticSourceControlUtils::RunGetHistory(Operation->ShouldUpdateHistory(), States, InCommand.ErrorMessages);
+
+			// Special case for conflicts
+			for (FPlasticSourceControlState& State : States)
 			{
-				for (int32 IdxFile = 0; IdxFile < States.Num(); IdxFile++)
+				if (State.IsConflicted())
 				{
-					FString& File = InCommand.Files[IdxFile];
-					FPlasticSourceControlState& State = States[IdxFile];
-
-					if (State.IsSourceControlled())
+					// In case of a merge conflict, we need to put the tip of the "remote branch" on top of the history
+					UE_LOG(LogSourceControl, Log, TEXT("%s: PendingMergeSourceChangeset %d"), *State.LocalFilename, State.PendingMergeSourceChangeset);
+					for (int32 IdxRevision = 0; IdxRevision < State.History.Num(); IdxRevision++)
 					{
-						// Get the history of the file (on all branches)
-						InCommand.bCommandSuccessful &= PlasticSourceControlUtils::RunGetHistory(File, InCommand.ErrorMessages, State);
-						if (State.IsConflicted())
+						const auto& Revision = State.History[IdxRevision];
+						if (Revision->ChangesetNumber == State.PendingMergeSourceChangeset)
 						{
-							// In case of a merge conflict, we need to put the tip of the "remote branch" on top of the history
-							UE_LOG(LogSourceControl, Log, TEXT("%s: PendingMergeSourceChangeset %d"), *State.LocalFilename, State.PendingMergeSourceChangeset);
-							for (int32 IdxRevision = 0; IdxRevision < State.History.Num(); IdxRevision++)
+							// If the Source Changeset is not already at the top of the History, duplicate it there.
+							if (IdxRevision > 0)
 							{
-								const auto& Revision = State.History[IdxRevision];
-								if (Revision->ChangesetNumber == State.PendingMergeSourceChangeset)
-								{
-									// If the Source Changeset is not already at the top of the History, duplicate it there.
-									if (IdxRevision > 0)
-									{
-										const auto RevisionCopy = Revision;
-										State.History.Insert(RevisionCopy, 0);
-									}
-									break;
-								}
+								const auto RevisionCopy = Revision;
+								State.History.Insert(RevisionCopy, 0);
 							}
+							break;
 						}
 					}
 				}
 			}
 		}
+		else
+		{
+			const FPlasticSourceControlModule& PlasticSourceControl = FModuleManager::GetModuleChecked<FPlasticSourceControlModule>("PlasticSourceControl");
+			if (PlasticSourceControl.AccessSettings().GetUpdateStatusOtherBranches() && AreAllFiles(InCommand.Files))
+			{
+				// Get only the last revision of the files (checking all branches)
+				// in order to warn the user if the file has been changed on another branch
+				InCommand.bCommandSuccessful &= PlasticSourceControlUtils::RunGetHistory(Operation->ShouldUpdateHistory(), States, InCommand.ErrorMessages);
+			}
+		}
 	}
-	else
+	// no path provided: only update the status of assets in Content/ directory if requested
+	// Perforce "opened files" are those that have been modified (or added/deleted): that is what we get with a simple status from the root
+	// This is called by the "CheckOut" Content Browser filter as well as our source control Refresh menu.
+	else if (Operation->ShouldGetOpenedOnly())
 	{
-		// no path provided: only update the status of assets in Content/ directory and also Config files
 		TArray<FString> ProjectDirs;
 		ProjectDirs.Add(FPaths::ConvertRelativePathToFull(FPaths::ProjectContentDir()));
-		ProjectDirs.Add(FPaths::ConvertRelativePathToFull(FPaths::ProjectConfigDir()));
-		InCommand.bCommandSuccessful = PlasticSourceControlUtils::RunUpdateStatus(ProjectDirs, false, InCommand.Concurrency, InCommand.ErrorMessages, States, InCommand.ChangesetNumber, InCommand.BranchName);
+		InCommand.bCommandSuccessful = PlasticSourceControlUtils::RunUpdateStatus(ProjectDirs, Operation->ShouldUpdateHistory(), InCommand.Concurrency, InCommand.ErrorMessages, States, InCommand.ChangesetNumber, InCommand.BranchName);
 	}
 
-	// don't use the ShouldUpdateModifiedState() hint here as it is specific to Perforce: the above normal Plastic status has already told us this information (like Git and Mercurial)
+	// TODO: re-evaluate how to optimize this heavy operation using some of these hints flags
+	// - ShouldGetOpenedOnly hint would be to call for all a whole workspace status update
+	// - ShouldUpdateModifiedState hint not used as the above normal Plastic status has already told us this information (like Git and Mercurial)
 
 	return InCommand.bCommandSuccessful;
 }
@@ -545,24 +573,24 @@ bool IsMoveOperation(const FString& InOrigin)
 			Promise->SetValue(MoveTemp(AssetsData));
 		});
 		const TArray<FAssetData> AssetsData = Promise->GetFuture().Get();
-		UE_LOG(LogSourceControl, Log, TEXT("PackageName: %s, AssetsData: Num=%d"), *PackageName, AssetsData.Num());
+		UE_LOG(LogSourceControl, Log, TEXT("IsMoveOperation: PackageName: %s, AssetsData: Num=%d"), *PackageName, AssetsData.Num());
 		if (AssetsData.Num() > 0)
 		{
 			const FAssetData& AssetData = AssetsData[0];
 			if (!AssetData.IsRedirector())
 			{
-				UE_LOG(LogSourceControl, Log, TEXT("%s is a plain asset, so it's a duplicate/copy"), *InOrigin);
+				UE_LOG(LogSourceControl, Log, TEXT("IsMoveOperation: %s is a plain asset, so it's a duplicate/copy"), *InOrigin);
 				bIsMoveOperation = false;
 			}
 			else
 			{
-				UE_LOG(LogSourceControl, Log, TEXT("%s is a redirector, so it's a move/rename"), *InOrigin);
+				UE_LOG(LogSourceControl, Log, TEXT("IsMoveOperation: %s is a redirector, so it's a move/rename"), *InOrigin);
 			}
 		}
 		else
 		{
 			// no asset in package (no redirector) so it should be a rename/move of a newly Added (not Controlled/Checked-In) file
-			UE_LOG(LogSourceControl, Log, TEXT("%s does not have asset in package (ie. no redirector) so it's a move/rename of a newly added file"), *InOrigin);
+			UE_LOG(LogSourceControl, Log, TEXT("IsMoveOperation: %s does not have asset in package (ie. no redirector) so it's a move/rename of a newly added file"), *InOrigin);
 		}
 	}
 
@@ -590,47 +618,30 @@ bool FPlasticCopyWorker::Execute(FPlasticSourceControlCommand& InCommand)
 		{
 			UE_LOG(LogSourceControl, Log, TEXT("Moving %s to %s..."), *Origin, *Destination);
 			// In case of rename, we have to undo what the Editor (created a redirector and added the dest asset), and then redo it with Plastic SCM
-			// - backup the redirector (if it exists) to a temp file
-			const bool bReplace = true;
-			const bool bEvenIfReadOnly = true;
-			const FString TempFileName = FPaths::CreateTempFilename(*FPaths::ProjectLogDir(), TEXT("Plastic-MoveTemp"), TEXT(".uasset"));
-			UE_LOG(LogSourceControl, Log, TEXT("Move '%s' -> '%d'"), *Origin, *TempFileName);
-			InCommand.bCommandSuccessful = IFileManager::Get().Move(*TempFileName, *Origin, bReplace, bEvenIfReadOnly);
 			// - revert the 'cm add' that was applied to the destination by the Editor
-			if (InCommand.bCommandSuccessful)
 			{
 				TArray<FString> DestinationFiles;
 				DestinationFiles.Add(Destination);
 				InCommand.bCommandSuccessful = PlasticSourceControlUtils::RunCommand(TEXT("undochange"), TArray<FString>(), DestinationFiles, InCommand.Concurrency, InCommand.InfoMessages, InCommand.ErrorMessages);
 			}
-			// - move back the asset from the destination to it's original location
-			if (InCommand.bCommandSuccessful)
-			{
-				UE_LOG(LogSourceControl, Log, TEXT("Move '%s' -> '%d'"), *Destination, *Origin);
-				InCommand.bCommandSuccessful = IFileManager::Get().Move(*Origin, *Destination, bReplace, bEvenIfReadOnly);
-			}
-			// - execute a 'cm move' command to the destination to redo the actual job
+			// - execute a 'cm move --nomoveondisk' command to the destination to tell cm what happened
 			if (InCommand.bCommandSuccessful)
 			{
+				TArray<FString> Parameters;
+				Parameters.Add(TEXT("--nomoveondisk"));
 				TArray<FString> Files;
 				Files.Add(Origin);
 				Files.Add(Destination);
 				// Detect special case for a partial checkout (CS:-1 in Gluon mode)!
 				if (-1 != InCommand.ChangesetNumber)
 				{
-					InCommand.bCommandSuccessful = PlasticSourceControlUtils::RunCommand(TEXT("move"), TArray<FString>(), Files, InCommand.Concurrency, InCommand.InfoMessages, InCommand.ErrorMessages);
+					InCommand.bCommandSuccessful = PlasticSourceControlUtils::RunCommand(TEXT("move"), Parameters, Files, InCommand.Concurrency, InCommand.InfoMessages, InCommand.ErrorMessages);
 				}
 				else
 				{
-					InCommand.bCommandSuccessful = PlasticSourceControlUtils::RunCommand(TEXT("partial move"), TArray<FString>(), Files, InCommand.Concurrency, InCommand.InfoMessages, InCommand.ErrorMessages);
+					InCommand.bCommandSuccessful = PlasticSourceControlUtils::RunCommand(TEXT("partial move"), Parameters, Files, InCommand.Concurrency, InCommand.InfoMessages, InCommand.ErrorMessages);
 				}
 			}
-			// - restore the redirector file (if it exists) to it's former location
-			if (InCommand.bCommandSuccessful)
-			{
-				UE_LOG(LogSourceControl, Log, TEXT("Move '%s' -> '%d'"), *TempFileName, *Origin);
-				InCommand.bCommandSuccessful = IFileManager::Get().Move(*Origin, *TempFileName, bReplace, bEvenIfReadOnly);
-			}
 			// - add the redirector file (if it exists) to source control
 			if (InCommand.bCommandSuccessful)
 			{
@@ -701,11 +712,10 @@ bool FPlasticSyncWorker::Execute(FPlasticSourceControlCommand& InCommand)
 		// detect the special case of a Sync of the root folder:
 		if ((InCommand.Files.Num() == 1) && (InCommand.Files.Last() == InCommand.PathToWorkspaceRoot))
 		{
-			// only update the status of assets in Content/ directory and also Config files
-			TArray<FString> ProjectDirs;
-			ProjectDirs.Add(FPaths::ConvertRelativePathToFull(FPaths::ProjectContentDir()));
-			ProjectDirs.Add(FPaths::ConvertRelativePathToFull(FPaths::ProjectConfigDir()));
-			PlasticSourceControlUtils::RunUpdateStatus(ProjectDirs, false, InCommand.Concurrency, InCommand.ErrorMessages, States, InCommand.ChangesetNumber, InCommand.BranchName);
+			// only update the status of assets in the Content directory
+			TArray<FString> ContentDir;
+			ContentDir.Add(FPaths::ConvertRelativePathToFull(FPaths::ProjectContentDir()));
+			PlasticSourceControlUtils::RunUpdateStatus(ContentDir, false, InCommand.Concurrency, InCommand.ErrorMessages, States, InCommand.ChangesetNumber, InCommand.BranchName);
 		}
 		// else: optim, no need to update the status of our files since this is done immediately after by the Editor
 	}
@@ -746,7 +756,6 @@ bool FPlasticResolveWorker::Execute(FPlasticSourceControlCommand& InCommand)
 		TArray<FString> OneFile;
 		OneFile.Add(State->PendingMergeFilename);
 
-		// @todo temporary debug log
 		UE_LOG(LogSourceControl, Log, TEXT("resolve %s"), *State->PendingMergeFilename);
 
 		// Mark the conflicted file as resolved
diff --git a/Engine/Plugins/Developer/PlasticSourceControl/Source/PlasticSourceControl/Private/PlasticSourceControlProvider.cpp b/Engine/Plugins/Developer/PlasticSourceControl/Source/PlasticSourceControl/Private/PlasticSourceControlProvider.cpp
--- a/Engine/Plugins/Developer/PlasticSourceControl/Source/PlasticSourceControl/Private/PlasticSourceControlProvider.cpp
+++ b/Engine/Plugins/Developer/PlasticSourceControl/Source/PlasticSourceControl/Private/PlasticSourceControlProvider.cpp
@@ -1,13 +1,16 @@
 // Copyright Epic Games, Inc. All Rights Reserved.
 
 #include "PlasticSourceControlProvider.h"
+
 #include "PlasticSourceControlCommand.h"
-#include "ISourceControlModule.h"
 #include "PlasticSourceControlModule.h"
-#include "PlasticSourceControlSettings.h"
 #include "PlasticSourceControlOperations.h"
+#include "PlasticSourceControlSettings.h"
+#include "PlasticSourceControlState.h"
 #include "PlasticSourceControlUtils.h"
 #include "SPlasticSourceControlSettings.h"
+
+#include "ISourceControlModule.h"
 #include "Logging/MessageLog.h"
 #include "ScopedSourceControlProgress.h"
 #include "SourceControlHelpers.h"
@@ -25,15 +28,23 @@ static FName ProviderName("Plastic SCM");
 void FPlasticSourceControlProvider::Init(bool bForceConnection)
 {
 	// Init() is called multiple times at startup: do not check Plastic SCM each time
-	if(!bPlasticAvailable)
+	if (!bPlasticAvailable)
 	{
 		const TSharedPtr<IPlugin> Plugin = IPluginManager::Get().FindPlugin(TEXT("PlasticSourceControl"));
-		if(Plugin.IsValid())
+		if (Plugin.IsValid())
 		{
-			UE_LOG(LogSourceControl, Log, TEXT("Plastic SCM plugin '%s'"), *(Plugin->GetDescriptor().VersionName));
+			PluginVersion = Plugin->GetDescriptor().VersionName;
+			UE_LOG(LogSourceControl, Log, TEXT("Plastic SCM plugin '%s'"), *PluginVersion);
 		}
 
 		CheckPlasticAvailability();
+
+		// Override the source control logs verbosity level if needed based on settings
+		FPlasticSourceControlModule& PlasticSourceControl = FModuleManager::LoadModuleChecked<FPlasticSourceControlModule>("PlasticSourceControl");
+		if (PlasticSourceControl.AccessSettings().GetEnableVerboseLogs())
+		{
+			PlasticSourceControlUtils::SwitchVerboseLogs(true);
+		}
 	}
 
 	// bForceConnection: not used anymore
@@ -43,17 +54,17 @@ void FPlasticSourceControlProvider::CheckPlasticAvailability()
 {
 	FPlasticSourceControlModule& PlasticSourceControl = FModuleManager::LoadModuleChecked<FPlasticSourceControlModule>("PlasticSourceControl");
 	FString PathToPlasticBinary = PlasticSourceControl.AccessSettings().GetBinaryPath();
-	if(PathToPlasticBinary.IsEmpty())
+	if (PathToPlasticBinary.IsEmpty())
 	{
 		// Try to find Plastic binary, and update settings accordingly
 		PathToPlasticBinary = PlasticSourceControlUtils::FindPlasticBinaryPath();
-		if(!PathToPlasticBinary.IsEmpty())
+		if (!PathToPlasticBinary.IsEmpty())
 		{
 			PlasticSourceControl.AccessSettings().SetBinaryPath(PathToPlasticBinary);
 		}
 	}
 
-	if(!PathToPlasticBinary.IsEmpty())
+	if (!PathToPlasticBinary.IsEmpty())
 	{
 		// Find the path to the root Plastic directory (if any, else uses the ProjectDir)
 		const FString PathToProjectDir = FPaths::ConvertRelativePathToFull(FPaths::ProjectDir());
@@ -61,14 +72,17 @@ void FPlasticSourceControlProvider::CheckPlasticAvailability()
 
 		// Launch the Plastic SCM cli shell on the background to issue all commands during this session
 		bPlasticAvailable = PlasticSourceControlUtils::LaunchBackgroundPlasticShell(PathToPlasticBinary, PathToWorkspaceRoot);
-		if(bPlasticAvailable)
+		if (bPlasticAvailable)
 		{
 			PlasticSourceControlUtils::GetPlasticScmVersion(PlasticScmVersion);
 
 			// Get user name (from the global Plastic SCM client config)
 			PlasticSourceControlUtils::GetUserName(UserName);
 
-			if(!bWorkspaceFound)
+			// Register Console Commands
+			PlasticSourceControlConsole.Register();
+
+			if (!bWorkspaceFound)
 			{
 				UE_LOG(LogSourceControl, Warning, TEXT("'%s' is not part of a Plastic workspace"), *FPaths::ProjectDir());
 			}
@@ -88,6 +102,8 @@ void FPlasticSourceControlProvider::Close()
 	PlasticSourceControlUtils::Terminate();
 	// Remove all extensions to the "Source Control" menu in the Editor Toolbar
 	PlasticSourceControlMenu.Unregister();
+	// Unregister Console Commands
+	PlasticSourceControlConsole.Unregister();
 
 	bServerAvailable = false;
 	bPlasticAvailable = false;
@@ -95,10 +111,10 @@ void FPlasticSourceControlProvider::Close()
 	UserName.Empty();
 }
 
-TSharedRef<FPlasticSourceControlState, ESPMode::ThreadSafe> FPlasticSourceControlProvider::GetStateInternal(const FString& Filename) const
+TSharedRef<FPlasticSourceControlState, ESPMode::ThreadSafe> FPlasticSourceControlProvider::GetStateInternal(const FString& InFilename) const
 {
-	TSharedRef<FPlasticSourceControlState, ESPMode::ThreadSafe>* State = StateCache.Find(Filename);
-	if(State != NULL)
+	TSharedRef<FPlasticSourceControlState, ESPMode::ThreadSafe>* State = StateCache.Find(InFilename);
+	if (State != NULL)
 	{
 		// found cached item
 		return (*State);
@@ -106,8 +122,8 @@ TSharedRef<FPlasticSourceControlState, ESPMode::ThreadSafe> FPlasticSourceContro
 	else
 	{
 		// cache an unknown state for this item
-		TSharedRef<FPlasticSourceControlState, ESPMode::ThreadSafe> NewState = MakeShareable( new FPlasticSourceControlState(Filename) );
-		StateCache.Add(Filename, NewState);
+		TSharedRef<FPlasticSourceControlState, ESPMode::ThreadSafe> NewState = MakeShareable(new FPlasticSourceControlState(FString(InFilename)));
+		StateCache.Add(InFilename, NewState);
 		return NewState;
 	}
 }
@@ -116,6 +132,7 @@ FText FPlasticSourceControlProvider::GetStatusText() const
 {
 	FFormatNamedArguments Args;
     Args.Add( TEXT("PlasticScmVersion"), FText::FromString(PlasticScmVersion) );
+    Args.Add( TEXT("PluginVersion"), FText::FromString(PluginVersion) );
 	Args.Add( TEXT("WorkspacePath"), FText::FromString(PathToWorkspaceRoot) );
 	Args.Add( TEXT("WorkspaceName"), FText::FromString(WorkspaceName) );
 	Args.Add( TEXT("BranchName"), FText::FromString(BranchName) );
@@ -130,7 +147,7 @@ FText FPlasticSourceControlProvider::GetStatusText() const
 	}
 	Args.Add( TEXT("UserName"), FText::FromString(UserName) );
 
-	return FText::Format( NSLOCTEXT("Status", "Provider: Plastic\nEnabledLabel", "Plastic SCM {PlasticScmVersion}\nWorkspace: {WorkspaceName} ({WorkspacePath})\n{BranchName}\nChangeset: {ChangesetNumber}\nUser: {UserName}"), Args );
+	return FText::Format( NSLOCTEXT("Status", "Provider: Plastic\nEnabledLabel", "Plastic SCM {PlasticScmVersion} (plugin v{PluginVersion})\nWorkspace: {WorkspaceName} ({WorkspacePath})\n{BranchName}\nChangeset: {ChangesetNumber}\nUser: {UserName}"), Args );
 }
 
 /** Quick check if source control is enabled */
@@ -152,22 +169,22 @@ const FName& FPlasticSourceControlProvider::GetName(void) const
 
 ECommandResult::Type FPlasticSourceControlProvider::GetState( const TArray<FString>& InFiles, TArray< TSharedRef<ISourceControlState, ESPMode::ThreadSafe> >& OutState, EStateCacheUsage::Type InStateCacheUsage )
 {
-	if(!IsEnabled())
+	if (!IsEnabled())
 	{
 		return ECommandResult::Failed;
 	}
 
-	TArray<FString> AbsoluteFiles = SourceControlHelpers::AbsoluteFilenames(InFiles);
+	const TArray<FString> AbsoluteFiles = SourceControlHelpers::AbsoluteFilenames(InFiles);
 
-	if(InStateCacheUsage == EStateCacheUsage::ForceUpdate)
+	if (InStateCacheUsage == EStateCacheUsage::ForceUpdate)
 	{
 		UE_LOG(LogSourceControl, Log, TEXT("GetState: ForceUpdate"));
 		Execute(ISourceControlOperation::Create<FUpdateStatus>(), AbsoluteFiles);
 	}
 
-	for(const auto& AbsoluteFile : AbsoluteFiles)
+	for (const FString& AbsoluteFile : AbsoluteFiles)
 	{
-		OutState.Add(GetStateInternal(*AbsoluteFile));
+		OutState.Add(GetStateInternal(AbsoluteFile));
 	}
 
 	return ECommandResult::Succeeded;
@@ -181,10 +198,10 @@ ECommandResult::Type FPlasticSourceControlProvider::GetState(const TArray<FSourc
 TArray<FSourceControlStateRef> FPlasticSourceControlProvider::GetCachedStateByPredicate(TFunctionRef<bool(const FSourceControlStateRef&)> Predicate) const
 {
 	TArray<FSourceControlStateRef> Result;
-	for(const auto& CacheItem : StateCache)
+	for (const auto& CacheItem : StateCache)
 	{
 		FSourceControlStateRef State = CacheItem.Value;
-		if(Predicate(State))
+		if (Predicate(State))
 		{
 			Result.Add(State);
 		}
@@ -209,17 +226,15 @@ void FPlasticSourceControlProvider::UnregisterSourceControlStateChanged_Handle(
 
 ECommandResult::Type FPlasticSourceControlProvider::Execute( const FSourceControlOperationRef& InOperation, FSourceControlChangelistPtr InChangelist, const TArray<FString>& InFiles, EConcurrency::Type InConcurrency, const FSourceControlOperationComplete& InOperationCompleteDelegate )
 {
-	if(!bWorkspaceFound && !(InOperation->GetName() == "Connect") && !(InOperation->GetName() == "MakeWorkspace"))
+	if (!bWorkspaceFound && !(InOperation->GetName() == "Connect") && !(InOperation->GetName() == "MakeWorkspace"))
 	{
 		UE_LOG(LogSourceControl, Warning, TEXT("'%s': only Connect operation allowed whithout a workspace"), *InOperation->GetName().ToString());
 		return ECommandResult::Failed;
 	}
 
-	TArray<FString> AbsoluteFiles = SourceControlHelpers::AbsoluteFilenames(InFiles);
-
 	// Query to see if we allow this operation
 	TSharedPtr<IPlasticSourceControlWorker, ESPMode::ThreadSafe> Worker = CreateWorker(InOperation->GetName());
-	if(!Worker.IsValid())
+	if (!Worker.IsValid())
 	{
 		// this operation is unsupported by this source control provider
 		FFormatNamedArguments Arguments;
@@ -230,11 +245,11 @@ ECommandResult::Type FPlasticSourceControlProvider::Execute( const FSourceContro
 	}
 
 	FPlasticSourceControlCommand* Command = new FPlasticSourceControlCommand(InOperation, Worker.ToSharedRef());
-	Command->Files = AbsoluteFiles;
+	Command->Files = SourceControlHelpers::AbsoluteFilenames(InFiles);
 	Command->OperationCompleteDelegate = InOperationCompleteDelegate;
 
 	// fire off operation
-	if(InConcurrency == EConcurrency::Synchronous)
+	if (InConcurrency == EConcurrency::Synchronous)
 	{
 		Command->bAutoDelete = false;
 
@@ -277,7 +292,7 @@ bool FPlasticSourceControlProvider::UsesCheckout() const
 TSharedPtr<IPlasticSourceControlWorker, ESPMode::ThreadSafe> FPlasticSourceControlProvider::CreateWorker(const FName& InOperationName) const
 {
 	const FGetPlasticSourceControlWorker* Operation = WorkersMap.Find(InOperationName);
-	if(Operation != nullptr)
+	if (Operation != nullptr)
 	{
 		return Operation->Execute();
 	}
@@ -294,12 +309,12 @@ void FPlasticSourceControlProvider::OutputCommandMessages(const FPlasticSourceCo
 {
 	FMessageLog SourceControlLog("SourceControl");
 
-	for(int32 ErrorIndex = 0; ErrorIndex < InCommand.ErrorMessages.Num(); ++ErrorIndex)
+	for (int32 ErrorIndex = 0; ErrorIndex < InCommand.ErrorMessages.Num(); ++ErrorIndex)
 	{
 		SourceControlLog.Error(FText::FromString(InCommand.ErrorMessages[ErrorIndex]));
 	}
 
-	for(int32 InfoIndex = 0; InfoIndex < InCommand.InfoMessages.Num(); ++InfoIndex)
+	for (int32 InfoIndex = 0; InfoIndex < InCommand.InfoMessages.Num(); ++InfoIndex)
 	{
 		SourceControlLog.Info(FText::FromString(InCommand.InfoMessages[InfoIndex]));
 	}
@@ -311,7 +326,7 @@ void FPlasticSourceControlProvider::UpdateWorkspaceStatus(const class FPlasticSo
 	{
 		// Is connection successful?
 		bServerAvailable = InCommand.bCommandSuccessful;
-		bWorkspaceFound = (InCommand.WorkspaceName.Len() > 0);
+		bWorkspaceFound = !InCommand.WorkspaceName.IsEmpty();
 
 		WorkspaceName = InCommand.WorkspaceName;
 		RepositoryName = InCommand.RepositoryName;
@@ -331,9 +346,12 @@ void FPlasticSourceControlProvider::UpdateWorkspaceStatus(const class FPlasticSo
 	}
 
 	// And for all operations running UpdateStatus, get Changeset and Branch informations:
-	if (!InCommand.BranchName.IsEmpty())
+	if (InCommand.ChangesetNumber != 0)
 	{
 		ChangesetNumber = InCommand.ChangesetNumber;
+	}
+	if (!InCommand.BranchName.IsEmpty())
+	{
 		BranchName = InCommand.BranchName;
 	}
 }
@@ -341,10 +359,10 @@ void FPlasticSourceControlProvider::UpdateWorkspaceStatus(const class FPlasticSo
 void FPlasticSourceControlProvider::Tick()
 {	
 	bool bStatesUpdated = false;
-	for(int32 CommandIndex = 0; CommandIndex < CommandQueue.Num(); ++CommandIndex)
+	for (int32 CommandIndex = 0; CommandIndex < CommandQueue.Num(); ++CommandIndex)
 	{
 		FPlasticSourceControlCommand& Command = *CommandQueue[CommandIndex];
-		if(Command.bExecuteProcessed)
+		if (Command.bExecuteProcessed)
 		{
 			// Remove command from the queue
 			CommandQueue.RemoveAt(CommandIndex);
@@ -358,12 +376,14 @@ void FPlasticSourceControlProvider::Tick()
 			// dump any messages to output log
 			OutputCommandMessages(Command);
 
+			UE_LOG(LogSourceControl, Verbose, TEXT("%s processed in %.3lfs"), *Command.Operation->GetName().ToString(), (FPlatformTime::Seconds() - Command.StartTimestamp));
+
 			// run the completion delegate callback if we have one bound
 			ECommandResult::Type Result = Command.bCommandSuccessful ? ECommandResult::Succeeded : ECommandResult::Failed;
 			Command.OperationCompleteDelegate.ExecuteIfBound(Command.Operation, Result);
 
 			// commands that are left in the array during a tick need to be deleted
-			if(Command.bAutoDelete)
+			if (Command.bAutoDelete)
 			{
 				// Only delete commands that are not running 'synchronously'
 				delete &Command;
@@ -375,7 +395,7 @@ void FPlasticSourceControlProvider::Tick()
 		}
 	}
 
-	if(bStatesUpdated)
+	if (bStatesUpdated)
 	{
 		OnSourceControlStateChanged.Broadcast();
 	}
@@ -411,7 +431,7 @@ ECommandResult::Type FPlasticSourceControlProvider::ExecuteSynchronousCommand(FP
 		IssueCommand( InCommand );
 
 		// ... then wait for its completion (thus making it synchronous)
-		while(!InCommand.bExecuteProcessed)
+		while (!InCommand.bExecuteProcessed)
 		{
 			// Tick the command queue and update progress.
 			Tick();
@@ -425,7 +445,7 @@ ECommandResult::Type FPlasticSourceControlProvider::ExecuteSynchronousCommand(FP
 		// always do one more Tick() to make sure the command queue is cleaned up.
 		Tick();
 
-		if(InCommand.bCommandSuccessful)
+		if (InCommand.bCommandSuccessful)
 		{
 			Result = ECommandResult::Succeeded;
 		}
@@ -453,7 +473,7 @@ ECommandResult::Type FPlasticSourceControlProvider::ExecuteSynchronousCommand(FP
 
 ECommandResult::Type FPlasticSourceControlProvider::IssueCommand(FPlasticSourceControlCommand& InCommand)
 {
-	if(GThreadPool != nullptr)
+	if (GThreadPool != nullptr)
 	{
 		// Queue this to our worker thread(s) for resolving
 		GThreadPool->AddQueuedWork(&InCommand);
diff --git a/Engine/Plugins/Developer/PlasticSourceControl/Source/PlasticSourceControl/Private/PlasticSourceControlProvider.h b/Engine/Plugins/Developer/PlasticSourceControl/Source/PlasticSourceControl/Private/PlasticSourceControlProvider.h
--- a/Engine/Plugins/Developer/PlasticSourceControl/Source/PlasticSourceControl/Private/PlasticSourceControlProvider.h
+++ b/Engine/Plugins/Developer/PlasticSourceControl/Source/PlasticSourceControl/Private/PlasticSourceControlProvider.h
@@ -5,9 +5,11 @@
 #include "CoreMinimal.h"
 #include "ISourceControlProvider.h"
 #include "IPlasticSourceControlWorker.h"
-#include "PlasticSourceControlState.h"
+#include "PlasticSourceControlConsole.h"
 #include "PlasticSourceControlMenu.h"
 
+class FPlasticSourceControlState;
+
 DECLARE_DELEGATE_RetVal(FPlasticSourceControlWorkerRef, FGetPlasticSourceControlWorker)
 
 class FPlasticSourceControlProvider : public ISourceControlProvider
@@ -29,9 +31,9 @@ public:
 	virtual bool IsEnabled() const override;
 	virtual bool IsAvailable() const override;
 	virtual const FName& GetName(void) const override;
-	virtual bool QueryStateBranchConfig(const FString& ConfigSrc, const FString& ConfigDest) { return false; }
-	virtual void RegisterStateBranches(const TArray<FString>& BranchNames, const FString& ContentRoot) {}
-	virtual int32 GetStateBranchIndex(const FString& InBranchName) const { return INDEX_NONE; }
+	virtual bool QueryStateBranchConfig(const FString& ConfigSrc, const FString& ConfigDest) override { return false; }
+	virtual void RegisterStateBranches(const TArray<FString>& BranchNames, const FString& ContentRoot) override {}
+	virtual int32 GetStateBranchIndex(const FString& InBranchName) const override { return INDEX_NONE; }
 	virtual ECommandResult::Type GetState( const TArray<FString>& InFiles, TArray<FSourceControlStateRef>& OutState, EStateCacheUsage::Type InStateCacheUsage ) override;
 	virtual ECommandResult::Type GetState(const TArray<FSourceControlChangelistRef>& InChangelists, TArray<FSourceControlChangelistStateRef>& OutState, EStateCacheUsage::Type InStateCacheUsage) override;
 	virtual TArray<FSourceControlStateRef> GetCachedStateByPredicate(TFunctionRef<bool(const FSourceControlStateRef&)> Predicate) const override;
@@ -112,7 +114,7 @@ public:
 	}
 
 	/** Helper function used to update state cache */
-	TSharedRef<FPlasticSourceControlState, ESPMode::ThreadSafe> GetStateInternal(const FString& Filename) const;
+	TSharedRef<FPlasticSourceControlState, ESPMode::ThreadSafe> GetStateInternal(const FString& InFilename) const;
 
 	/**
 	 * Register a worker with the provider.
@@ -151,6 +153,9 @@ private:
     /** Version of the Plastic SCM executable used */
     FString PlasticScmVersion;
 
+	/** Version of the Plastic SCM plugin */
+	FString PluginVersion;
+
 	/** Path to the root of the Plastic workspace: can be the GameDir itself, or any parent directory (found by the "Connect" operation) */
 	FString PathToWorkspaceRoot;
 
@@ -173,7 +178,7 @@ private:
 	int32 ChangesetNumber;
 
 	/** State cache */
-	mutable TMap<FString, TSharedRef<class FPlasticSourceControlState, ESPMode::ThreadSafe> > StateCache;
+	mutable TMap<FString, TSharedRef<FPlasticSourceControlState, ESPMode::ThreadSafe> > StateCache;
 
 	/** The currently registered source control operations */
 	TMap<FName, FGetPlasticSourceControlWorker> WorkersMap;
@@ -186,4 +191,7 @@ private:
 
 	/** Source Control Menu Extension */
 	FPlasticSourceControlMenu PlasticSourceControlMenu;
+
+	/** Source Control Console commands */
+	FPlasticSourceControlConsole PlasticSourceControlConsole;
 };
diff --git a/Engine/Plugins/Developer/PlasticSourceControl/Source/PlasticSourceControl/Private/PlasticSourceControlRevision.cpp b/Engine/Plugins/Developer/PlasticSourceControl/Source/PlasticSourceControl/Private/PlasticSourceControlRevision.cpp
--- a/Engine/Plugins/Developer/PlasticSourceControl/Source/PlasticSourceControl/Private/PlasticSourceControlRevision.cpp
+++ b/Engine/Plugins/Developer/PlasticSourceControl/Source/PlasticSourceControl/Private/PlasticSourceControlRevision.cpp
@@ -28,24 +28,24 @@ bool FPlasticSourceControlRevision::Get(FString& InOutFilename, EConcurrency::Ty
 	}
 
 	// if a filename for the temp file wasn't supplied generate a unique-ish one
-	if(InOutFilename.Len() == 0)
+	if (InOutFilename.Len() == 0)
 	{
 		// create the diff dir if we don't already have it (Plastic wont)
 		IFileManager::Get().MakeDirectory(*FPaths::DiffDir(), true);
 		// create a unique temp file name based on the unique revision Id
-		const FString TempFileName = FString::Printf(TEXT("%stemp-%d-%s"), *FPaths::DiffDir(), RevisionId, *FPaths::GetCleanFilename(Filename));
+		const FString TempFileName = FString::Printf(TEXT("%stemp-%d-%s"), *FPaths::DiffDir(), ChangesetNumber, *FPaths::GetCleanFilename(Filename));
 		InOutFilename = FPaths::ConvertRelativePathToFull(TempFileName);
 	}
 
 	bool bCommandSuccessful;
-	if(FPaths::FileExists(InOutFilename))
+	if (FPaths::FileExists(InOutFilename))
 	{
 		bCommandSuccessful = true; // if the temp file already exists, reuse it directly
 	}
 	else if (State)
 	{
-		// Format the revision specification of the file, like revid:1230@repo@server:8087
-		const FString RevisionSpecification = FString::Printf(TEXT("revid:%d@%s"), RevisionId, *State->RepSpec);
+		// Format the revision specification of the file, like rev:Content/BP.uasset#cs:12@repo@server:8087
+		const FString RevisionSpecification = FString::Printf(TEXT("rev:%s#cs:%d@%s"), *Filename, ChangesetNumber, *State->RepSpec);
 		bCommandSuccessful = PlasticSourceControlUtils::RunDumpToFile(PathToPlasticBinary, RevisionSpecification, InOutFilename);
 	}
 	else
@@ -92,8 +92,8 @@ const FString& FPlasticSourceControlRevision::GetUserName() const
 
 const FString& FPlasticSourceControlRevision::GetClientSpec() const
 {
-	static FString EmptyString(TEXT("")); // NOTE Workspace/Clientspec of the submitter (Perforce only)
-	return EmptyString;
+	// Note: show Branch instead of the Workspace of the submitter since it's Perforce only
+	return Branch;
 }
 
 const FString& FPlasticSourceControlRevision::GetAction() const
diff --git a/Engine/Plugins/Developer/PlasticSourceControl/Source/PlasticSourceControl/Private/PlasticSourceControlRevision.h b/Engine/Plugins/Developer/PlasticSourceControl/Source/PlasticSourceControl/Private/PlasticSourceControlRevision.h
--- a/Engine/Plugins/Developer/PlasticSourceControl/Source/PlasticSourceControl/Private/PlasticSourceControlRevision.h
+++ b/Engine/Plugins/Developer/PlasticSourceControl/Source/PlasticSourceControl/Private/PlasticSourceControlRevision.h
@@ -3,7 +3,6 @@
 #pragma once
 
 #include "CoreMinimal.h"
-#include "ISourceControlProvider.h"
 #include "ISourceControlRevision.h"
 
 class FPlasticSourceControlState;
@@ -59,6 +58,9 @@ public:
 	/** The user that made the change */
 	FString UserName;
 
+	/** Branch where the change was made */
+	FString Branch;
+
 	/** The action (add, edit, branch etc.) performed at this revision */
 	FString Action;
 
diff --git a/Engine/Plugins/Developer/PlasticSourceControl/Source/PlasticSourceControl/Private/PlasticSourceControlSettings.cpp b/Engine/Plugins/Developer/PlasticSourceControl/Source/PlasticSourceControl/Private/PlasticSourceControlSettings.cpp
--- a/Engine/Plugins/Developer/PlasticSourceControl/Source/PlasticSourceControl/Private/PlasticSourceControlSettings.cpp
+++ b/Engine/Plugins/Developer/PlasticSourceControl/Source/PlasticSourceControl/Private/PlasticSourceControlSettings.cpp
@@ -28,14 +28,14 @@ bool FPlasticSourceControlSettings::SetBinaryPath(const FString& InString)
 {
 	FScopeLock ScopeLock(&CriticalSection);
 	const bool bChanged = (BinaryPath != InString);
-	if(bChanged)
+	if (bChanged)
 	{
 		BinaryPath = InString;
 	}
 	return bChanged;
 }
 
-bool FPlasticSourceControlSettings::UpdateStatusAtStartup() const
+bool FPlasticSourceControlSettings::GetUpdateStatusAtStartup() const
 {
 	FScopeLock ScopeLock(&CriticalSection);
 	return bUpdateStatusAtStartup;
@@ -47,6 +47,30 @@ void FPlasticSourceControlSettings::SetUpdateStatusAtStartup(const bool bInUpdat
 	bUpdateStatusAtStartup = bInUpdateStatusAtStartup;
 }
 
+bool FPlasticSourceControlSettings::GetUpdateStatusOtherBranches() const
+{
+	FScopeLock ScopeLock(&CriticalSection);
+	return bUpdateStatusOtherBranches;
+}
+
+void FPlasticSourceControlSettings::SetUpdateStatusOtherBranches(const bool bInUpdateStatusOtherBranches)
+{
+	FScopeLock ScopeLock(&CriticalSection);
+	bUpdateStatusOtherBranches = bInUpdateStatusOtherBranches;
+}
+
+bool FPlasticSourceControlSettings::GetEnableVerboseLogs() const
+{
+	FScopeLock ScopeLock(&CriticalSection);
+	return bEnableVerboseLogs;
+}
+
+void FPlasticSourceControlSettings::SetEnableVerboseLogs(const bool bInEnableVerboseLogs)
+{
+	FScopeLock ScopeLock(&CriticalSection);
+	bEnableVerboseLogs = bInEnableVerboseLogs;
+}
+
 // This is called at startup nearly before anything else in our module: BinaryPath will then be used by the provider
 void FPlasticSourceControlSettings::LoadSettings()
 {
@@ -54,6 +78,8 @@ void FPlasticSourceControlSettings::LoadSettings()
 	const FString& IniFile = SourceControlHelpers::GetSettingsIni();
 	GConfig->GetString(*PlasticSettingsConstants::SettingsSection, TEXT("BinaryPath"), BinaryPath, IniFile);
 	GConfig->GetBool(*PlasticSettingsConstants::SettingsSection, TEXT("UpdateStatusAtStartup"), bUpdateStatusAtStartup, IniFile);
+	GConfig->GetBool(*PlasticSettingsConstants::SettingsSection, TEXT("UpdateStatusOtherBranches"), bUpdateStatusOtherBranches, IniFile);
+	GConfig->GetBool(*PlasticSettingsConstants::SettingsSection, TEXT("EnableVerboseLogs"), bEnableVerboseLogs, IniFile);
 }
 
 void FPlasticSourceControlSettings::SaveSettings() const
@@ -62,4 +88,6 @@ void FPlasticSourceControlSettings::SaveSettings() const
 	const FString& IniFile = SourceControlHelpers::GetSettingsIni();
 	GConfig->SetString(*PlasticSettingsConstants::SettingsSection, TEXT("BinaryPath"), *BinaryPath, IniFile);
 	GConfig->SetBool(*PlasticSettingsConstants::SettingsSection, TEXT("UpdateStatusAtStartup"), bUpdateStatusAtStartup, IniFile);
-}
\ No newline at end of file
+	GConfig->SetBool(*PlasticSettingsConstants::SettingsSection, TEXT("UpdateStatusOtherBranches"), bUpdateStatusOtherBranches, IniFile);
+	GConfig->SetBool(*PlasticSettingsConstants::SettingsSection, TEXT("EnableVerboseLogs"), bEnableVerboseLogs, IniFile);
+}
diff --git a/Engine/Plugins/Developer/PlasticSourceControl/Source/PlasticSourceControl/Private/PlasticSourceControlSettings.h b/Engine/Plugins/Developer/PlasticSourceControl/Source/PlasticSourceControl/Private/PlasticSourceControlSettings.h
--- a/Engine/Plugins/Developer/PlasticSourceControl/Source/PlasticSourceControl/Private/PlasticSourceControlSettings.h
+++ b/Engine/Plugins/Developer/PlasticSourceControl/Source/PlasticSourceControl/Private/PlasticSourceControlSettings.h
@@ -12,9 +12,17 @@ public:
 	bool SetBinaryPath(const FString& InString);
 
 	/** Enable an asynchronous "Update Status" at Editor Startup (default is no, can take a long time). */
-	bool UpdateStatusAtStartup() const;
+	bool GetUpdateStatusAtStartup() const;
 	void SetUpdateStatusAtStartup(const bool bInUpdateStatusAtStartup);
 
+	/** Enable Update status to call "history" to detect recent changesets on other branches (can be slow). */
+	bool GetUpdateStatusOtherBranches() const;
+	void SetUpdateStatusOtherBranches(const bool bInUpdateStatusOtherBranches);
+
+	/** Enable LogSourceControl Verbose logs */
+	bool GetEnableVerboseLogs() const;
+	void SetEnableVerboseLogs(const bool bInEnableVerboseLogs);
+
 	/** Load settings from ini file */
 	void LoadSettings();
 
@@ -26,11 +34,17 @@ private:
 	mutable FCriticalSection CriticalSection;
 
 	/** Plastic binary path */
-	FString BinaryPath;
+	FString BinaryPath = TEXT("cm");
 
-	/** Enable an asynchronous "Update Status" at Editor Startup (default is no).
+	/** Run an asynchronous "Update Status" at Editor Startup (default is no).
 	 * This does not work well with very big projects where this operation could take dozens of seconds
 	 * preventing the project to have any source control support during this time.
 	*/
-	bool bUpdateStatusAtStartup;
+	bool bUpdateStatusAtStartup = false;
+
+	/** Enable Update status to call "history" to detect recent changesets on other branches (can be slow). */
+	bool bUpdateStatusOtherBranches = false;
+
+	/** Override LogSourceControl verbosity level to Verbose, and back, if not already VeryVerbose */
+	bool bEnableVerboseLogs = false;
 };
\ No newline at end of file
diff --git a/Engine/Plugins/Developer/PlasticSourceControl/Source/PlasticSourceControl/Private/PlasticSourceControlState.cpp b/Engine/Plugins/Developer/PlasticSourceControl/Source/PlasticSourceControl/Private/PlasticSourceControlState.cpp
--- a/Engine/Plugins/Developer/PlasticSourceControl/Source/PlasticSourceControl/Private/PlasticSourceControlState.cpp
+++ b/Engine/Plugins/Developer/PlasticSourceControl/Source/PlasticSourceControl/Private/PlasticSourceControlState.cpp
@@ -47,9 +47,9 @@ TSharedPtr<class ISourceControlRevision, ESPMode::ThreadSafe> FPlasticSourceCont
 
 TSharedPtr<class ISourceControlRevision, ESPMode::ThreadSafe> FPlasticSourceControlState::FindHistoryRevision( int32 RevisionNumber ) const
 {
-	for(const auto& Revision : History)
+	for (const auto& Revision : History)
 	{
-		if(Revision->GetRevisionNumber() == RevisionNumber)
+		if (Revision->GetRevisionNumber() == RevisionNumber)
 		{
 			return Revision;
 		}
@@ -60,9 +60,9 @@ TSharedPtr<class ISourceControlRevision, ESPMode::ThreadSafe> FPlasticSourceCont
 
 TSharedPtr<class ISourceControlRevision, ESPMode::ThreadSafe> FPlasticSourceControlState::FindHistoryRevision(const FString& InRevision) const
 {
-	for(const auto& Revision : History)
+	for (const auto& Revision : History)
 	{
-		if(Revision->GetRevision() == InRevision)
+		if (Revision->GetRevision() == InRevision)
 		{
 			return Revision;
 		}
@@ -73,10 +73,10 @@ TSharedPtr<class ISourceControlRevision, ESPMode::ThreadSafe> FPlasticSourceCont
 
 TSharedPtr<class ISourceControlRevision, ESPMode::ThreadSafe> FPlasticSourceControlState::GetBaseRevForMerge() const
 {
-	for(const auto& Revision : History)
+	for (const auto& Revision : History)
 	{
 		// look for the the SHA1 id of the file, not the commit id (revision)
-		if(Revision->ChangesetNumber == PendingMergeBaseChangeset)
+		if (Revision->ChangesetNumber == PendingMergeBaseChangeset)
 		{
 			return Revision;
 		}
@@ -91,26 +91,33 @@ FSlateIcon FPlasticSourceControlState::GetIcon() const
 	{
 		return FSlateIcon(FAppStyle::GetAppStyleSetName(),"Perforce.NotAtHeadRevision");
 	}
+	else if (WorkspaceState != EWorkspaceState::CheckedOut && WorkspaceState != EWorkspaceState::LockedByOther)
+	{
+		if (IsModifiedInOtherBranch())
+		{
+			return FSlateIcon(FAppStyle::GetAppStyleSetName(), "Perforce.ModifiedOtherBranch");
+		}
+	}
 
 	switch (WorkspaceState)
 	{
 	case EWorkspaceState::CheckedOut:
-	case EWorkspaceState::Replaced: // Merged (waiting for checkin) TODO: would need a dedicated icon
+	case EWorkspaceState::Replaced: // Merged (waiting for checkin)
 		return FSlateIcon(FAppStyle::GetAppStyleSetName(), "Perforce.CheckedOut");
 	case EWorkspaceState::Added:
 	case EWorkspaceState::Copied:
 		return FSlateIcon(FAppStyle::GetAppStyleSetName(), "Perforce.OpenForAdd");
 	case EWorkspaceState::Moved:
 		return FSlateIcon(FAppStyle::GetAppStyleSetName(), "Perforce.Branched");
-	case EWorkspaceState::Deleted: // Deleted & Missing files does not show in Content Browser
+	case EWorkspaceState::Deleted:
+	case EWorkspaceState::LocallyDeleted:
 		return FSlateIcon(FAppStyle::GetAppStyleSetName(), "Perforce.MarkedForDelete");
-	case EWorkspaceState::LocallyDeleted: // Deleted & Missing files does not show in Content Browser
 	case EWorkspaceState::Conflicted:
 		return FSlateIcon(FAppStyle::GetAppStyleSetName(), "Perforce.NotAtHeadRevision");
 	case EWorkspaceState::LockedByOther:
 		return FSlateIcon(FAppStyle::GetAppStyleSetName(), "Perforce.CheckedOutByOtherUser", NAME_None, "SourceControl.LockOverlay");
 	case EWorkspaceState::Private: // Not controlled
-	case EWorkspaceState::Changed: // Changed but unchecked-out file is in a certain way not controlled - TODO: would need a dedicated icon
+	case EWorkspaceState::Changed: // Changed but unchecked-out file is in a certain way not controlled
 		return FSlateIcon(FAppStyle::GetAppStyleSetName(), "Perforce.NotInDepot");
 	case EWorkspaceState::Unknown:
 	case EWorkspaceState::Ignored:
@@ -122,7 +129,20 @@ FSlateIcon FPlasticSourceControlState::GetIcon() const
 
 FText FPlasticSourceControlState::GetDisplayName() const
 {
-	switch(WorkspaceState)
+	if (!IsCurrent())
+	{
+		return LOCTEXT("NotCurrent", "Not current");
+	}
+	else if (WorkspaceState != EWorkspaceState::LockedByOther)
+	{
+		if (IsModifiedInOtherBranch())
+		{
+			FNumberFormattingOptions NoCommas;
+			NoCommas.UseGrouping = false;
+			return FText::Format(LOCTEXT("ModifiedOtherBranch", "Modified in {0} CS:{1} by {2}"), FText::FromString(HeadBranch), FText::AsNumber(HeadChangeList, &NoCommas), FText::FromString(HeadUserName));
+		}
+	}
+	switch (WorkspaceState)
 	{
 	case EWorkspaceState::Unknown:
 		return LOCTEXT("Unknown", "Unknown");
@@ -131,7 +151,7 @@ FText FPlasticSourceControlState::GetDisplayName() const
 	case EWorkspaceState::Controlled:
 		return LOCTEXT("Controlled", "Controlled");
 	case EWorkspaceState::CheckedOut:
-		return LOCTEXT("CheckedOut", "CheckedOut");
+		return LOCTEXT("CheckedOut", "Checked-out");
 	case EWorkspaceState::Added:
 		return LOCTEXT("Added", "Added");
 	case EWorkspaceState::Moved:
@@ -147,9 +167,9 @@ FText FPlasticSourceControlState::GetDisplayName() const
 	case EWorkspaceState::Changed:
 		return LOCTEXT("Changed", "Changed");
 	case EWorkspaceState::Conflicted:
-		return LOCTEXT("ContentsConflict", "Contents Conflict");
+		return LOCTEXT("ContentsConflict", "Conflicted");
 	case EWorkspaceState::LockedByOther:
-		return FText::Format(LOCTEXT("CheckedOutOther", "Checked out by: {0}"), FText::FromString(LockedBy));
+		return FText::Format(LOCTEXT("CheckedOutOther", "Checked out by: {0} in {1}"), FText::FromString(LockedBy), FText::FromString(LockedWhere));
 	case EWorkspaceState::Private:
 		return LOCTEXT("NotControlled", "Not Under Source Control");
 	}
@@ -159,34 +179,50 @@ FText FPlasticSourceControlState::GetDisplayName() const
 
 FText FPlasticSourceControlState::GetDisplayTooltip() const
 {
-	switch(WorkspaceState)
+	if (!IsCurrent())
+	{
+		return LOCTEXT("NotCurrent_Tooltip", "Not at the head revision");
+	}
+	else if (WorkspaceState != EWorkspaceState::LockedByOther)
+	{
+		if (IsModifiedInOtherBranch())
+		{
+			FNumberFormattingOptions NoCommas;
+			NoCommas.UseGrouping = false;
+			return FText::Format(LOCTEXT("ModifiedOtherBranch_Tooltip", "Modified in {0} CS:{1} by {2}"), FText::FromString(HeadBranch), FText::AsNumber(HeadChangeList, &NoCommas), FText::FromString(HeadUserName));
+		}
+	}
+
+	switch (WorkspaceState)
 	{
 	case EWorkspaceState::Unknown:
-		return LOCTEXT("Unknown_Tooltip", "Unknown source control state");
+		return LOCTEXT("Unknown_Tooltip", "The file status is unknown");
 	case EWorkspaceState::Ignored:
-		return LOCTEXT("Ignored_Tooltip", "Item is being ignored.");
+		return LOCTEXT("Ignored_Tooltip", "The file is ignored");
 	case EWorkspaceState::Controlled:
-		return LOCTEXT("Pristine_Tooltip", "There are no modifications");
+		return LOCTEXT("Pristine_Tooltip", "The file has no modification");
 	case EWorkspaceState::CheckedOut:
-		return LOCTEXT("CheckedOut_Tooltip", "Item is checked out");
+		return LOCTEXT("CheckedOut_Tooltip", "The file is checked out");
 	case EWorkspaceState::Added:
-		return LOCTEXT("Added_Tooltip", "Item has been added");
+		return LOCTEXT("Added_Tooltip", "The file has been added");
 	case EWorkspaceState::Moved:
-		return LOCTEXT("Moved_Tooltip", "Item has been moved or renamed");
+		return LOCTEXT("Moved_Tooltip", "The file has been moved or renamed");
 	case EWorkspaceState::Copied:
-		return LOCTEXT("Copied_Tooltip", "Item has been copied");
+		return LOCTEXT("Copied_Tooltip", "The file has been copied");
 	case EWorkspaceState::Replaced:
-		return LOCTEXT("Replaced_Tooltip", "Item has been replaced / merged");
+		return LOCTEXT("Replaced_Tooltip", "The file has been replaced / merged");
 	case EWorkspaceState::Deleted:
-		return LOCTEXT("Deleted_Tooltip", "Item is scheduled for deletion");
+		return LOCTEXT("Deleted_Tooltip", "The file has been deleted");
 	case EWorkspaceState::LocallyDeleted:
-		return LOCTEXT("LocallyDeleted_Tooltip", "Item is missing");
+		return LOCTEXT("LocallyDeleted_Tooltip", "The file is missing");
 	case EWorkspaceState::Changed:
-		return LOCTEXT("Modified_Tooltip", "Item has been modified");
+		return LOCTEXT("Modified_Tooltip", "The file has been modified");
 	case EWorkspaceState::Conflicted:
-		return LOCTEXT("ContentsConflict_Tooltip", "The contents of the item conflict with updates received from the repository.");
+		return LOCTEXT("ContentsConflict_Tooltip", "The content of the file conflict with updates received from the repository");
+	case EWorkspaceState::LockedByOther:
+		return FText::Format(LOCTEXT("CheckedOutOther_Tooltip", "Checked out by {0} in {1}"), FText::FromString(LockedBy), FText::FromString(LockedWhere));
 	case EWorkspaceState::Private:
-		return LOCTEXT("NotControlled_Tooltip", "Item is not under version control.");
+		return LOCTEXT("NotControlled_Tooltip", "The file is not under version control");
 	}
 
 	return FText();
@@ -209,16 +245,21 @@ const FDateTime& FPlasticSourceControlState::GetTimeStamp() const
 // Deleted and Missing assets cannot appear in the Content Browser but does appear in Submit to Source Control Window
 bool FPlasticSourceControlState::CanCheckIn() const
 {
-	const bool bCanCheckIn = WorkspaceState == EWorkspaceState::Added
+	const bool bCanCheckIn =
+		(  WorkspaceState == EWorkspaceState::Added
 		|| WorkspaceState == EWorkspaceState::Deleted
 		|| WorkspaceState == EWorkspaceState::LocallyDeleted
-		|| WorkspaceState == EWorkspaceState::Changed // NOTE: Comment to enable checkout on prompt of a Changed file (see bellow)
+		|| WorkspaceState == EWorkspaceState::Changed
 		|| WorkspaceState == EWorkspaceState::Moved
 		|| WorkspaceState == EWorkspaceState::Copied
 		|| WorkspaceState == EWorkspaceState::Replaced
-		|| WorkspaceState == EWorkspaceState::CheckedOut;
+		|| WorkspaceState == EWorkspaceState::CheckedOut
+		) && !IsConflicted() && IsCurrent();
 
-	UE_LOG(LogSourceControl, Log, TEXT("%s CanCheckIn=%d"), *LocalFilename, bCanCheckIn);
+	if (!IsUnknown())
+	{
+		UE_LOG(LogSourceControl, Verbose, TEXT("%s CanCheckIn=%d"), *LocalFilename, bCanCheckIn);
+	}
 
 	return bCanCheckIn;
 }
@@ -230,7 +271,10 @@ bool FPlasticSourceControlState::CanCheckout() const
 								|| WorkspaceState == EWorkspaceState::Replaced)		// In source control, merged, waiting for checkin to conclude the merge 
 								&& IsCurrent(); // Is up to date (at the revision of the repo)
 
-	UE_LOG(LogSourceControl, Log, TEXT("%s CanCheckout=%d"), *LocalFilename, bCanCheckout);
+	if (!IsUnknown())
+	{
+		UE_LOG(LogSourceControl, Verbose, TEXT("%s CanCheckout=%d"), *LocalFilename, bCanCheckout);
+	}
 
 	return bCanCheckout;
 }
@@ -239,11 +283,14 @@ bool FPlasticSourceControlState::IsCheckedOut() const
 {
 	const bool bIsCheckedOut = WorkspaceState == EWorkspaceState::CheckedOut
 							|| WorkspaceState == EWorkspaceState::Moved
-							|| WorkspaceState == EWorkspaceState::Conflicted
-							|| WorkspaceState == EWorkspaceState::Replaced // Workaround to enable checkin
-							|| WorkspaceState == EWorkspaceState::Changed; // Workaround to enable checkin  NOTE: Comment to enable checkout on prompt of a Changed file (see above)
+							|| WorkspaceState == EWorkspaceState::Conflicted	// In source control, waiting for merged
+							|| WorkspaceState == EWorkspaceState::Replaced		// In source control, merged, waiting for checkin to conclude the merge 
+							|| WorkspaceState == EWorkspaceState::Changed;		// Note: Workaround to enable checkin (still required by UE5.0)
 
-	if (bIsCheckedOut) UE_LOG(LogSourceControl, Log, TEXT("%s IsCheckedOut"), *LocalFilename);
+	if (bIsCheckedOut)
+	{
+		UE_LOG(LogSourceControl, Verbose, TEXT("%s IsCheckedOut"), *LocalFilename);
+	}
 
 	return bIsCheckedOut;
 }
@@ -256,41 +303,27 @@ bool FPlasticSourceControlState::IsCheckedOutOther(FString* Who) const
 	}
 	const bool bIsLockedByOther = WorkspaceState == EWorkspaceState::LockedByOther;
 
-	// @todo: temporary debug log
-	// if (bIsLockedByOther) UE_LOG(LogSourceControl, Log, TEXT("%s IsCheckedOutOther by '%s' (%s)"), *LocalFilename, *LockedBy, *LockedWhere);
+	if (bIsLockedByOther)
+	{
+		UE_LOG(LogSourceControl, VeryVerbose, TEXT("%s IsCheckedOutOther by '%s' (%s)"), *LocalFilename, *LockedBy, *LockedWhere);
+	}
 
 	return bIsLockedByOther;
 }
 
 
 /** Get whether this file is checked out in a different branch, if no branch is specified defaults to FEngineVerion current branch */
-bool FPlasticSourceControlState::IsCheckedOutInOtherBranch(const FString& CurrentBranch) const
+bool FPlasticSourceControlState::IsCheckedOutInOtherBranch(const FString& CurrentBranch /* = FString() */) const
 {
+	// Note: to my knowledge, it's not possible to detect that with PlasticSCM withouth the Locks,
+	// which are already detected by fileinfo LockedBy/LockedWhere and reported by IsCheckedOutOther() above
 	return false;
 }
 
 /** Get whether this file is modified in a different branch, if no branch is specified defaults to FEngineVerion current branch */
-bool FPlasticSourceControlState::IsModifiedInOtherBranch(const FString& CurrentBranch) const
-{
-	return false;
-}
-
-/** Get whether this file is checked out or modified in a different branch, if no branch is specified defaults to FEngineVerion current branch */
-bool FPlasticSourceControlState::IsCheckedOutOrModifiedInOtherBranch(const FString& CurrentBranch) const
+bool FPlasticSourceControlState::IsModifiedInOtherBranch(const FString& CurrentBranch /* = FString() */) const
 {
-	return false;
-}
-
-/** Get the other branches this file is checked out in */
-TArray<FString> FPlasticSourceControlState::GetCheckedOutBranches() const
-{
-	return TArray<FString>();
-}
-
-/** Get the user info for checkouts on other branches */
-FString FPlasticSourceControlState::GetOtherUserBranchCheckedOuts() const
-{
-	return FString();
+	return !HeadBranch.IsEmpty();
 }
 
 /** Get head modification information for other branches
@@ -298,17 +331,22 @@ FString FPlasticSourceControlState::GetOtherUserBranchCheckedOuts() const
 */
 bool FPlasticSourceControlState::GetOtherBranchHeadModification(FString& HeadBranchOut, FString& ActionOut, int32& HeadChangeListOut) const
 {
-	return false;
-}
+	HeadBranchOut = HeadBranch;
+	ActionOut = HeadAction;
+	HeadChangeListOut = HeadChangeList;
 
+	return !HeadBranch.IsEmpty();
+}
 
 bool FPlasticSourceControlState::IsCurrent() const
 {
 	// NOTE: Deleted assets get a "-1" HeadRevision which we do not want to override the real icon state
 	const bool bIsCurrent = (LocalRevisionChangeset == DepotRevisionChangeset) || (WorkspaceState == EWorkspaceState::Deleted);
 
-	// @todo: temporary debug log
-	// if (bIsCurrent) UE_LOG(LogSourceControl, Log, TEXT("%s IsCurrent"), *LocalFilename, );
+	if (bIsCurrent)
+	{
+		UE_LOG(LogSourceControl, VeryVerbose, TEXT("%s IsCurrent"), *LocalFilename);
+	}
 	
 	return bIsCurrent;
 }
@@ -378,7 +416,7 @@ bool FPlasticSourceControlState::IsModified() const
 {
 	// Warning: for a clean "check-in" (commit) checked-out files unmodified should be removed from the changeset (the index)
 	//
-	// Thus, before check-in UE Editor call RevertUnchangedFiles() in PromptForCheckin() and CheckinFiles().
+	// Thus, before check-in UE4 Editor call RevertUnchangedFiles() in PromptForCheckin() and CheckinFiles().
 	//
 	// So here we must take care to enumerate all states that need to be commited, all other will be discarded:
 	//  - Unknown
@@ -395,7 +433,7 @@ bool FPlasticSourceControlState::IsModified() const
 							|| WorkspaceState == EWorkspaceState::Changed
 							|| WorkspaceState == EWorkspaceState::Conflicted;
 
-	UE_LOG(LogSourceControl, Log, TEXT("%s IsModified=%d"), *LocalFilename, bIsModified);
+	UE_LOG(LogSourceControl, Verbose, TEXT("%s IsModified=%d"), *LocalFilename, bIsModified);
 
 	return bIsModified;
 }
diff --git a/Engine/Plugins/Developer/PlasticSourceControl/Source/PlasticSourceControl/Private/PlasticSourceControlState.h b/Engine/Plugins/Developer/PlasticSourceControl/Source/PlasticSourceControl/Private/PlasticSourceControlState.h
--- a/Engine/Plugins/Developer/PlasticSourceControl/Source/PlasticSourceControl/Private/PlasticSourceControlState.h
+++ b/Engine/Plugins/Developer/PlasticSourceControl/Source/PlasticSourceControl/Private/PlasticSourceControlState.h
@@ -34,19 +34,24 @@ class FPlasticSourceControlState : public ISourceControlState, public TSharedFro
 {
 public:
 
-	FPlasticSourceControlState(const FString& InLocalFilename)
-		: LocalFilename(InLocalFilename)
+	FPlasticSourceControlState(FString&& InLocalFilename)
+		: LocalFilename(MoveTemp(InLocalFilename))
 		, WorkspaceState(EWorkspaceState::Unknown)
-		, DepotRevisionChangeset(-1)
-		, LocalRevisionChangeset(-1)
+		, DepotRevisionChangeset(INVALID_REVISION)
+		, LocalRevisionChangeset(INVALID_REVISION)
 		, TimeStamp(0)
 	{
 	}
 
+	FPlasticSourceControlState() = delete;
 	FPlasticSourceControlState(const FPlasticSourceControlState& InState) = delete;
 	const FPlasticSourceControlState& operator=(const FPlasticSourceControlState& InState) = delete;
 
 	FPlasticSourceControlState(FPlasticSourceControlState&& InState)
+		: WorkspaceState(EWorkspaceState::Unknown)
+		, DepotRevisionChangeset(INVALID_REVISION)
+		, LocalRevisionChangeset(INVALID_REVISION)
+		, TimeStamp(0)
 	{
 		Move(MoveTemp(InState));
 	}
@@ -61,16 +66,32 @@ public:
 	{
 		History = MoveTemp(InState.History);
 		LocalFilename = MoveTemp(InState.LocalFilename);
-		RepSpec = MoveTemp(InState.RepSpec);
+		WorkspaceState = InState.WorkspaceState;
 		PendingMergeFilename = MoveTemp(InState.PendingMergeFilename);
 		PendingMergeBaseChangeset = InState.PendingMergeBaseChangeset;
 		PendingMergeSourceChangeset = InState.PendingMergeSourceChangeset;
 		PendingMergeParameters = MoveTemp(InState.PendingMergeParameters);
-		LockedBy = MoveTemp(InState.LockedBy);
-		LockedWhere = MoveTemp(InState.LockedWhere);
-		WorkspaceState = InState.WorkspaceState;
-		DepotRevisionChangeset = InState.DepotRevisionChangeset;
-		LocalRevisionChangeset = InState.LocalRevisionChangeset;
+		// Update "fileinfo" information only if the command was issued
+		if (InState.DepotRevisionChangeset != INVALID_REVISION)
+		{
+			LockedBy = MoveTemp(InState.LockedBy);
+			LockedWhere = MoveTemp(InState.LockedWhere);
+			RepSpec = MoveTemp(InState.RepSpec);
+			DepotRevisionChangeset = InState.DepotRevisionChangeset;
+			LocalRevisionChangeset = InState.LocalRevisionChangeset;
+
+			HeadBranch = MoveTemp(InState.HeadBranch);
+			HeadAction = MoveTemp(InState.HeadAction);
+			HeadChangeList = MoveTemp(InState.HeadChangeList);
+			HeadUserName = MoveTemp(InState.HeadUserName);
+			HeadModTime = MoveTemp(InState.HeadModTime);
+		}
+		// Don't override "fileinfo" information in case of an optimized/lightweight "whole folder status" triggered by a global Submit Content or Refresh
+		// and regenerate the LockedByOther state based on LockedBy
+		else if (!IsCheckedOut() && !LockedBy.IsEmpty())
+		{
+			WorkspaceState = EWorkspaceState::LockedByOther;
+		}
 		MovedFrom = MoveTemp(InState.MovedFrom);
 		TimeStamp = InState.TimeStamp;
 
@@ -102,12 +123,12 @@ public:
 	virtual bool CanCheckout() const override;
 	virtual bool IsCheckedOut() const override;
 	virtual bool IsCheckedOutOther(FString* Who = nullptr) const override;
-	virtual bool IsCheckedOutInOtherBranch(const FString& CurrentBranch = FString()) const;
-	virtual bool IsModifiedInOtherBranch(const FString& CurrentBranch = FString()) const;
-	virtual bool IsCheckedOutOrModifiedInOtherBranch(const FString& CurrentBranch = FString()) const;
-	virtual TArray<FString> GetCheckedOutBranches() const;
-	virtual FString GetOtherUserBranchCheckedOuts() const;
-	virtual bool GetOtherBranchHeadModification(FString& HeadBranchOut, FString& ActionOut, int32& HeadChangeListOut) const;
+	virtual bool IsCheckedOutInOtherBranch(const FString& CurrentBranch = FString()) const override;
+	virtual bool IsModifiedInOtherBranch(const FString& CurrentBranch = FString()) const override;
+	virtual bool IsCheckedOutOrModifiedInOtherBranch(const FString& CurrentBranch = FString()) const override { return IsCheckedOutInOtherBranch(CurrentBranch) || IsModifiedInOtherBranch(CurrentBranch); }
+	virtual TArray<FString> GetCheckedOutBranches() const override { return TArray<FString>(); }
+	virtual FString GetOtherUserBranchCheckedOuts() const override { return FString(); }
+	virtual bool GetOtherBranchHeadModification(FString& HeadBranchOut, FString& ActionOut, int32& HeadChangeListOut) const override;
 	virtual bool IsCurrent() const override;
 	virtual bool IsSourceControlled() const override;
 	virtual bool IsAdded() const override;
@@ -119,7 +140,7 @@ public:
 	virtual bool CanAdd() const override;
 	virtual bool CanDelete() const override;
 	virtual bool IsConflicted() const override;
-	virtual bool CanRevert() const;
+	virtual bool CanRevert() const override;
 
 public:
 	/** History of the item, if any */
@@ -161,9 +182,21 @@ public:
 	/** Original name in case of a Moved/Renamed file */
 	FString MovedFrom;
 
-	/** Whether the file is a binary file or not */
-//	bool bBinary;
-
 	/** The timestamp of the last update */
 	FDateTime TimeStamp;
+
+	/** The branch with the head change list */
+	FString HeadBranch;
+
+	/** The type of action of the last modification */
+	FString HeadAction;
+
+	/** The user of the last modification */
+	FString HeadUserName;
+
+	/** The last file modification time */
+	int64 HeadModTime;
+
+	/** The change list of the last modification */
+	int32 HeadChangeList;
 };
diff --git a/Engine/Plugins/Developer/PlasticSourceControl/Source/PlasticSourceControl/Private/PlasticSourceControlUtils.cpp b/Engine/Plugins/Developer/PlasticSourceControl/Source/PlasticSourceControl/Private/PlasticSourceControlUtils.cpp
--- a/Engine/Plugins/Developer/PlasticSourceControl/Source/PlasticSourceControl/Private/PlasticSourceControlUtils.cpp
+++ b/Engine/Plugins/Developer/PlasticSourceControl/Source/PlasticSourceControl/Private/PlasticSourceControlUtils.cpp
@@ -38,7 +38,7 @@ namespace PlasticSourceControlConstants
 FScopedTempFile::FScopedTempFile(const FText& InText)
 {
 	Filename = FPaths::CreateTempFilename(*FPaths::ProjectLogDir(), TEXT("Plastic-Temp"), TEXT(".txt"));
-	if(!FFileHelper::SaveStringToFile(InText.ToString(), *Filename, FFileHelper::EEncodingOptions::ForceUTF8WithoutBOM))
+	if (!FFileHelper::SaveStringToFile(InText.ToString(), *Filename, FFileHelper::EEncodingOptions::ForceUTF8WithoutBOM))
 	{
 		UE_LOG(LogSourceControl, Error, TEXT("Failed to write to temp file: %s"), *Filename);
 	}
@@ -46,9 +46,9 @@ FScopedTempFile::FScopedTempFile(const FText& InText)
 
 FScopedTempFile::~FScopedTempFile()
 {
-	if(FPaths::FileExists(Filename))
+	if (FPaths::FileExists(Filename))
 	{
-		if(!FPlatformFileManager::Get().GetPlatformFile().DeleteFile(*Filename))
+		if (!FPlatformFileManager::Get().GetPlatformFile().DeleteFile(*Filename))
 		{
 			UE_LOG(LogSourceControl, Error, TEXT("Failed to delete temp file: %s"), *Filename);
 		}
@@ -60,35 +60,13 @@ const FString& FScopedTempFile::GetFilename() const
 	return Filename;
 }
 
-// Needed to SetHandleInformation() on WritePipe for input (opposite of ReadPipe, for output) (idem FInteractiveProcess)
-static FORCEINLINE bool CreatePipeWrite(void*& ReadPipe, void*& WritePipe)
-{
-#if PLATFORM_WINDOWS
-	SECURITY_ATTRIBUTES Attr = { sizeof(SECURITY_ATTRIBUTES), NULL, true };
-
-	if (!::CreatePipe(&ReadPipe, &WritePipe, &Attr, 0))
-	{
-		return false;
-	}
-
-	if (!::SetHandleInformation(WritePipe, HANDLE_FLAG_INHERIT, 0))
-	{
-		return false;
-	}
-
-	return true;
-#else
-	return FPlatformProcess::CreatePipe(ReadPipe, WritePipe);
-#endif // PLATFORM_WINDOWS
-}
-
 namespace PlasticSourceControlUtils
 {
-// Command-line interface paramters and output format changed with version 8.0.16.3000
-// For more details, see https://www.plasticscm.com/download/releasenotes/oldernotes?release=8.0.16.3000
+// Command-line interface parameters and output format changed with version 8.0.16.3000
+// For more details, see https://www.plasticscm.com/download/releasenotes/8.0.16.3000
 static bool				bIsNewVersion80163000 = false;
 
-// In/Out Pipes for the 'cm shell' persistent process
+// In/Out Pipes for the 'cm shell' persistent child process
 static void*			ShellOutputPipeRead = nullptr;
 static void*			ShellOutputPipeWrite = nullptr;
 static void*			ShellInputPipeRead = nullptr;
@@ -96,12 +74,13 @@ static void*			ShellInputPipeWrite = nullptr;
 static FProcHandle		ShellProcessHandle;
 static FCriticalSection	ShellCriticalSection;
 static size_t			ShellCommandCounter = -1;
+static double			ShellCumulatedTime = 0.;
 
 // Internal function to cleanup (called under the critical section)
 static void _CleanupBackgroundCommandLineShell()
 {
-	FPlatformProcess::ClosePipe(ShellInputPipeRead, ShellInputPipeWrite);
 	FPlatformProcess::ClosePipe(ShellOutputPipeRead, ShellOutputPipeWrite);
+	FPlatformProcess::ClosePipe(ShellInputPipeRead, ShellInputPipeWrite);
 	ShellOutputPipeRead = ShellOutputPipeWrite = nullptr;
 	ShellInputPipeRead = ShellInputPipeWrite = nullptr;
 }
@@ -115,8 +94,10 @@ static bool _StartBackgroundPlasticShell(const FString& InPathToPlasticBinary, c
 	const bool bLaunchHidden = true;				// the new process will be minimized in the task bar
 	const bool bLaunchReallyHidden = bLaunchHidden; // the new process will not have a window or be in the task bar
 
-	verify(FPlatformProcess::CreatePipe(ShellOutputPipeRead, ShellOutputPipeWrite));	// For reading from child process
-	verify(CreatePipeWrite(ShellInputPipeRead, ShellInputPipeWrite));	// For writing to child process
+	const double StartTimestamp = FPlatformTime::Seconds();
+
+	verify(FPlatformProcess::CreatePipe(ShellOutputPipeRead, ShellOutputPipeWrite, false));	// For reading outputs from cm shell child process
+	verify(FPlatformProcess::CreatePipe(ShellInputPipeRead, ShellInputPipeWrite, true));	// For writing commands to cm shell child process
 
 	ShellProcessHandle = FPlatformProcess::CreateProc(*InPathToPlasticBinary, *FullCommand, bLaunchDetached, bLaunchHidden, bLaunchReallyHidden, nullptr, 0, *InWorkingDirectory, ShellOutputPipeWrite, ShellInputPipeRead);
 	if (!ShellProcessHandle.IsValid())
@@ -126,13 +107,42 @@ static bool _StartBackgroundPlasticShell(const FString& InPathToPlasticBinary, c
 	}
 	else
 	{
-		UE_LOG(LogSourceControl, Log, TEXT("LaunchBackgroundPlasticShell: '%s %s' ok (handle %d)"), *InPathToPlasticBinary, *FullCommand, ShellProcessHandle.Get());
+		const double ElapsedTime = (FPlatformTime::Seconds() - StartTimestamp);
+		UE_LOG(LogSourceControl, Verbose, TEXT("_StartBackgroundPlasticShell: '%s %s' ok (in %.3lfs, handle %d)"), *InPathToPlasticBinary, *FullCommand, ElapsedTime, ShellProcessHandle.Get());
 		ShellCommandCounter = 0;
+		ShellCumulatedTime = ElapsedTime;
 	}
 
 	return ShellProcessHandle.IsValid();
 }
 
+// Internal function (called under the critical section)
+static void _ExitBackgroundCommandLineShell()
+{
+	if (ShellProcessHandle.IsValid())
+	{
+		if (FPlatformProcess::IsProcRunning(ShellProcessHandle))
+		{
+			// Tell the 'cm shell' to exit
+			FPlatformProcess::WritePipe(ShellInputPipeWrite, TEXT("exit"));
+			// And wait up to one second for its termination
+			const double Timeout = 1.0;
+			const double StartTimestamp = FPlatformTime::Seconds();
+			while (FPlatformProcess::IsProcRunning(ShellProcessHandle))
+			{
+				if ((FPlatformTime::Seconds() - StartTimestamp) > Timeout)
+				{
+					UE_LOG(LogSourceControl, Warning, TEXT("ExitBackgroundCommandLineShell: cm shell didn't stop gracefuly in %lfs."), Timeout);
+					break;
+				}
+				FPlatformProcess::Sleep(0.01f);
+			}
+		}
+		FPlatformProcess::CloseProc(ShellProcessHandle);
+		_CleanupBackgroundCommandLineShell();
+	}
+}
+
 // Internal function (called under the critical section)
 static void _RestartBackgroundCommandLineShell()
 {
@@ -140,8 +150,7 @@ static void _RestartBackgroundCommandLineShell()
 	const FString& PathToPlasticBinary = PlasticSourceControl.AccessSettings().GetBinaryPath();
 	const FString& WorkingDirectory = PlasticSourceControl.GetProvider().GetPathToWorkspaceRoot();
 
-	FPlatformProcess::CloseProc(ShellProcessHandle);
-	_CleanupBackgroundCommandLineShell();
+	_ExitBackgroundCommandLineShell();
 	_StartBackgroundPlasticShell(PathToPlasticBinary, WorkingDirectory);
 }
 
@@ -150,15 +159,15 @@ static bool _RunCommandInternal(const FString& InCommand, const TArray<FString>&
 {
 	bool bResult = false;
 
+	ShellCommandCounter++;
+
 	// Detect previous crash of cm.exe and restart 'cm shell'
 	if (!FPlatformProcess::IsProcRunning(ShellProcessHandle))
 	{
-		UE_LOG(LogSourceControl, Warning, TEXT("RunCommand(%d): 'cm shell' has stopped. Restarting!"), ShellCommandCounter);
+		UE_LOG(LogSourceControl, Warning, TEXT("RunCommand: 'cm shell' has stopped. Restarting! (count %d)"), ShellCommandCounter);
 		_RestartBackgroundCommandLineShell();
 	}
 
-	ShellCommandCounter++;
-
 	// Start with the Plastic command itself ("status", "log", "chekin"...)
 	FString FullCommand = InCommand;
 	// Append to the command all parameters, and then finally the files
@@ -173,9 +182,8 @@ static bool _RunCommandInternal(const FString& InCommand, const TArray<FString>&
 		FullCommand += File;
 		FullCommand += TEXT("\"");
 	}
-	// @todo: temporary debug logs (without the end-of-line)
 	const FString LoggableCommand = FullCommand.Left(256); // Limit command log size to 256 characters
-	UE_LOG(LogSourceControl, Log, TEXT("RunCommand(%d): '%s' (%d chars, %d files)"), ShellCommandCounter, *LoggableCommand, FullCommand.Len()+1, InFiles.Num());
+	UE_LOG(LogSourceControl, Verbose, TEXT("RunCommand: '%s' (%d chars, %d files)"), *LoggableCommand, FullCommand.Len()+1, InFiles.Num());
 	FullCommand += TEXT('\n'); // Finalize the command line
 
 	// Send command to 'cm shell' process
@@ -191,7 +199,7 @@ static bool _RunCommandInternal(const FString& InCommand, const TArray<FString>&
 	while (FPlatformProcess::IsProcRunning(ShellProcessHandle))
 	{
 		FString Output = FPlatformProcess::ReadPipe(ShellOutputPipeRead);
-		if (0 < Output.Len())
+		if (!Output.IsEmpty())
 		{
 			LastActivity = FPlatformTime::Seconds(); // freshen the timestamp while cm is still actively outputting information
 			OutResults.Append(MoveTemp(Output));
@@ -204,7 +212,7 @@ static bool _RunCommandInternal(const FString& InCommand, const TArray<FString>&
 				{
 					const FString Result = OutResults.Mid(IndexCommandResult + 14, IndexEndResult - IndexCommandResult - 14);
 					const int32 ResultCode = FCString::Atoi(*Result);
-					bResult = (0 == ResultCode);
+					bResult = (ResultCode == 0);
 					// remove the CommandResult line from the OutResults
 					OutResults.RemoveAt(IndexCommandResult, OutResults.Len() - IndexCommandResult);
 					break;
@@ -215,41 +223,50 @@ static bool _RunCommandInternal(const FString& InCommand, const TArray<FString>&
 		{
 			// In case of long running operation, start to print intermediate output from cm shell (like percentage of progress)
 			// (but only when running Asynchronous commands, since Synchronous commands block the main thread until they finish)
-			UE_LOG(LogSourceControl, Log, TEXT("RunCommand(%d): '%s' in progress for %lfs...\n%s"), ShellCommandCounter, *InCommand, (FPlatformTime::Seconds() - StartTimestamp), *OutResults.Mid(PreviousLogLen));
+			UE_LOG(LogSourceControl, Log, TEXT("RunCommand: '%s' in progress for %.3lfs...\n%s"), *InCommand, (FPlatformTime::Seconds() - StartTimestamp), *OutResults.Mid(PreviousLogLen));
 			PreviousLogLen = OutResults.Len();
 			LastLog = FPlatformTime::Seconds(); // freshen the timestamp of last log
 		}
 		else if (FPlatformTime::Seconds() - LastActivity > Timeout)
 		{
-			// In case of timeout, ask the blocking 'cm shell' process to exit, and detach from it immediatly: it will be relaunched by next command
-			UE_LOG(LogSourceControl, Error, TEXT("RunCommand(%d): '%s' %d TIMEOUT after %lfs output:\n%s"), ShellCommandCounter, *InCommand, bResult, (FPlatformTime::Seconds() - StartTimestamp), *OutResults.Mid(PreviousLogLen));
-			FPlatformProcess::WritePipe(ShellInputPipeWrite, TEXT("exit"));
-			FPlatformProcess::CloseProc(ShellProcessHandle);
-			_CleanupBackgroundCommandLineShell();
+			// In case of timeout, ask the blocking 'cm shell' process to exit, detach from it and restart it immediatly
+			UE_LOG(LogSourceControl, Error, TEXT("RunCommand: '%s' TIMEOUT after %.3lfs output (%d chars):\n%s"), *InCommand, (FPlatformTime::Seconds() - StartTimestamp), OutResults.Len(), *OutResults.Mid(PreviousLogLen));
+			_RestartBackgroundCommandLineShell();
+			return false;
 		}
 
 		FPlatformProcess::Sleep(0.001f);
 	}
+	const double ElapsedTime = (FPlatformTime::Seconds() - StartTimestamp);
+
 	if (!InCommand.Equals(TEXT("exit")))
 	{
 		if (!FPlatformProcess::IsProcRunning(ShellProcessHandle))
 		{
 			// 'cm shell' normally only terminates in case of 'exit' command. Will restart on next command.
-			UE_LOG(LogSourceControl, Error, TEXT("RunCommand(%d): '%s' 'cm shell' stopped after %lfs output:\n%s"), ShellCommandCounter, *LoggableCommand, (FPlatformTime::Seconds() - StartTimestamp), *OutResults.Left(4096)); // Limit result size to 4096 characters
+			UE_LOG(LogSourceControl, Error, TEXT("RunCommand: '%s' 'cm shell' stopped after %.3lfs output (%d chars):\n%s"), *LoggableCommand, ElapsedTime, OutResults.Len(), *OutResults.Left(4096)); // Limit result size to 4096 characters
 		}
 		else if (!bResult)
 		{
-			UE_LOG(LogSourceControl, Warning, TEXT("RunCommand(%d): '%s' (in %lfs) output:\n%s"), ShellCommandCounter, *LoggableCommand, (FPlatformTime::Seconds() - StartTimestamp), *OutResults.Left(4096)); // Limit result size to 4096 characters
+			UE_LOG(LogSourceControl, Warning, TEXT("RunCommand: '%s' (in %.3lfs) output (%d chars):\n%s"), *LoggableCommand, ElapsedTime, OutResults.Len(), *OutResults.Left(4096)); // Limit result size to 4096 characters
 		}
 		else
 		{
 			if (PreviousLogLen > 0)
 			{
-				UE_LOG(LogSourceControl, Log, TEXT("RunCommand(%d): '%s' (in %lfs) output:\n%s"), ShellCommandCounter, *LoggableCommand, (FPlatformTime::Seconds() - StartTimestamp), *OutResults.Mid(PreviousLogLen).Left(4096)); // Limit result size to 4096 characters
+				UE_LOG(LogSourceControl, Log, TEXT("RunCommand: '%s' (in %.3lfs) output (%d chars):\n%s"), *LoggableCommand, ElapsedTime, OutResults.Len(), *OutResults.Mid(PreviousLogLen).Left(4096)); // Limit result size to 4096 characters
 			}
 			else
 			{
-				UE_LOG(LogSourceControl, Log, TEXT("RunCommand(%d): '%s' (in %lfs) output:\n%s"), ShellCommandCounter, *LoggableCommand, (FPlatformTime::Seconds() - StartTimestamp), *OutResults.Left(4096)); // Limit result size to 4096 characters
+				if (OutResults.Len() <= 200) // Limit result size to 200 characters
+				{
+					UE_LOG(LogSourceControl, Log, TEXT("RunCommand: '%s' (in %.3lfs) output (%d chars):\n%s"), *LoggableCommand, ElapsedTime, OutResults.Len(), *OutResults);
+				}
+				else
+				{
+					UE_LOG(LogSourceControl, Log, TEXT("RunCommand: '%s' (in %.3lfs) (output %d chars not displayed)"), *LoggableCommand, ElapsedTime, OutResults.Len());
+					UE_LOG(LogSourceControl, Verbose, TEXT("\n%s"), *OutResults.Left(4096));; // Limit result size to 4096 characters
+				}
 			}
 		}
 	}
@@ -259,23 +276,10 @@ static bool _RunCommandInternal(const FString& InCommand, const TArray<FString>&
 		OutErrors = MoveTemp(OutResults);
 	}
 
-	return bResult;
-}
+	ShellCumulatedTime += ElapsedTime;
+	UE_LOG(LogSourceControl, Verbose, TEXT("RunCommand: cumulated time spent in shell: %.3lfs (count %d)"), ShellCumulatedTime, ShellCommandCounter);
 
-// Internal function (called under the critical section)
-static void _ExitBackgroundCommandLineShell()
-{
-	// Tell the 'cm shell' to exit
-	FString Results, Errors;
-	_RunCommandInternal(TEXT("exit"), TArray<FString>(), TArray<FString>(), EConcurrency::Synchronous, Results, Errors);
-	// And wait up to one seconde for its termination
-	int timeout = 100;
-	while (FPlatformProcess::IsProcRunning(ShellProcessHandle) && (0 < timeout--))
-	{
-		FPlatformProcess::Sleep(0.01f);
-	}
-	FPlatformProcess::CloseProc(ShellProcessHandle);
-	_CleanupBackgroundCommandLineShell();
+	return bResult;
 }
 
 // Launch the Plastic SCM background 'cm shell' process in background for optimized successive commands (thread-safe)
@@ -285,10 +289,7 @@ bool LaunchBackgroundPlasticShell(const FString& InPathToPlasticBinary, const FS
 	FScopeLock Lock(&ShellCriticalSection);
 
 	// terminate previous shell if one is already running
-	if (ShellProcessHandle.IsValid())
-	{
-		_ExitBackgroundCommandLineShell();
-	}
+	_ExitBackgroundCommandLineShell();
 
 	return _StartBackgroundPlasticShell(InPathToPlasticBinary, InWorkingDirectory);
 }
@@ -299,31 +300,16 @@ void Terminate()
 	// Protect public APIs from multi-thread access
 	FScopeLock Lock(&ShellCriticalSection);
 
-	if (ShellProcessHandle.IsValid())
-	{
-		_ExitBackgroundCommandLineShell();
-	}
+	_ExitBackgroundCommandLineShell();
 }
 
 // Run command (thread-safe)
-static bool RunCommandInternal(const FString& InCommand, const TArray<FString>& InParameters, const TArray<FString>& InFiles, const EConcurrency::Type InConcurrency, FString& OutResults, FString& OutErrors)
+bool RunCommandInternal(const FString& InCommand, const TArray<FString>& InParameters, const TArray<FString>& InFiles, const EConcurrency::Type InConcurrency, FString& OutResults, FString& OutErrors)
 {
-	bool bResult = false;
-
 	// Protect public APIs from multi-thread access
 	FScopeLock Lock(&ShellCriticalSection);
 
-	if (ShellProcessHandle.IsValid())
-	{
-		bResult = _RunCommandInternal(InCommand, InParameters, InFiles, InConcurrency, OutResults, OutErrors);
-	}
-	else
-	{
-		UE_LOG(LogSourceControl, Error, TEXT("RunCommand(%d): '%s': cm shell not running"), ShellCommandCounter, *InCommand);
-		OutErrors = InCommand + ": Plastic SCM shell not running!";
-	}
-
-	return bResult;
+	return _RunCommandInternal(InCommand, InParameters, InFiles, InConcurrency, OutResults, OutErrors);
 }
 
 // Basic parsing or results & errors from the Plastic command line process
@@ -359,7 +345,7 @@ bool FindRootDirectory(const FString& InPath, FString& OutWorkspaceRoot)
 	auto TrimTrailing = [](FString& Str, const TCHAR Char)
 	{
 		int32 Len = Str.Len();
-		while(Len && Str[Len - 1] == Char)
+		while (Len && Str[Len - 1] == Char)
 		{
 			Str = Str.LeftChop(1);
 			Len = Str.Len();
@@ -369,15 +355,15 @@ bool FindRootDirectory(const FString& InPath, FString& OutWorkspaceRoot)
 	TrimTrailing(OutWorkspaceRoot, '\\');
 	TrimTrailing(OutWorkspaceRoot, '/');
 
-	while(!bFound && !OutWorkspaceRoot.IsEmpty())
+	while (!bFound && !OutWorkspaceRoot.IsEmpty())
 	{
 		// Look for the ".plastic" subdirectory present at the root of every Plastic workspace
 		PathToPlasticSubdirectory = OutWorkspaceRoot / TEXT(".plastic");
-		bFound = IFileManager::Get().DirectoryExists(*PathToPlasticSubdirectory);
-		if(!bFound)
+		bFound = FPaths::DirectoryExists(*PathToPlasticSubdirectory);
+		if (!bFound)
 		{
 			int32 LastSlashIndex;
-			if(OutWorkspaceRoot.FindLastChar(TEXT('/'), LastSlashIndex))
+			if (OutWorkspaceRoot.FindLastChar(TEXT('/'), LastSlashIndex))
 			{
 				OutWorkspaceRoot = OutWorkspaceRoot.Left(LastSlashIndex);
 			}
@@ -434,8 +420,7 @@ void GetPlasticScmVersion(FString& OutPlasticScmVersion)
 	{
 		OutPlasticScmVersion = InfoMessages[0];
 
-		// Command-line format output changed with version 8.0.16.3000, plugin will crash unless we demand the old format.
-		// For more details, see https://www.plasticscm.com/download/releasenotes/oldernotes?release=8.0.16.3000
+		// Command-line format output changed with version 8.0.16.3000
 		bIsNewVersion80163000 = !PlasticScmVersionLess(OutPlasticScmVersion, "8.0.16.3000");
 	}
 }
@@ -501,7 +486,13 @@ static bool ParseWorkspaceInformation(const TArray<FString>& InInfoMessages, int
 	// Get the branch name, in the form "Branch /main@UEPlasticPluginDev" (enabled by the "--wkconfig" flag)
 	if (InInfoMessages.Num() > 1)
 	{
-		OutBranchName = InInfoMessages[1];
+		static const FString BranchPrefix(TEXT("Branch "));
+		const FString& BranchInfo = InInfoMessages[1];
+		const int32 BranchIndex = BranchInfo.Find(BranchPrefix, ESearchCase::CaseSensitive);
+		if (BranchIndex > INDEX_NONE)
+		{
+			OutBranchName = BranchInfo;
+		}
 	}
 
 	return bResult;
@@ -513,13 +504,18 @@ bool GetWorkspaceInformation(int32& OutChangeset, FString& OutRepositoryName, FS
 	TArray<FString> ErrorMessages;
 	TArray<FString> Parameters;
 
-	// Command-line format output changed with version 8.0.16.3000, plugin will crash unless we demand the old format.
-	// For more details, see https://www.plasticscm.com/download/releasenotes/oldernotes?release=8.0.16.3000
-	if (bIsNewVersion80163000) {
+	// Command-line format output changed with version 8.0.16.3000, see https://www.plasticscm.com/download/releasenotes/8.0.16.3000
+	if (bIsNewVersion80163000)
+	{
 		Parameters.Add(TEXT("--compact"));
+		Parameters.Add(TEXT("--header")); // Only prints the workspace status. No file status.
 	}
-	Parameters.Add(TEXT("--wkconfig")); // Branch name. TODO Deprecated in 8.0.16.3000 https://www.plasticscm.com/download/releasenotes/oldernotes?release=8.0.16.3000 => by default in xml format
-	Parameters.Add(TEXT("--nochanges")); // No file status // TODO: The option "--nochanges" has been renamed to "--header", and outputs just the workspace status.
+	else
+	{
+		Parameters.Add(TEXT("--nochanges")); // Only prints the workspace status. No file status.
+	}
+	// NOTE: --wkconfig results in two network calls GetBranchInfoByName & GetLastChangesetOnBranch so it's okay to do it once here but not all the time
+	Parameters.Add(TEXT("--wkconfig")); // Branch name. NOTE: Deprecated in 8.0.16.3000 https://www.plasticscm.com/download/releasenotes/8.0.16.3000
 	bool bResult = RunCommand(TEXT("status"), Parameters, TArray<FString>(), EConcurrency::Synchronous, InfoMessages, ErrorMessages);
 	if (bResult)
 	{
@@ -676,7 +672,7 @@ static EWorkspaceState::Type StateFromPlasticStatus(const FString& InResult)
 }
 
 /**
- * @brief Parse the array of strings results of a 'cm status --noheaders --controlledchanged --changed --localdeleted --localmoved --private --ignored' command
+ * @brief Parse the array of strings results of a 'cm status --noheaders --all --ignored' command
  *
  * Called in case of a regular status command for one or multiple files (not for a whole directory). 
  *
@@ -697,19 +693,20 @@ static EWorkspaceState::Type StateFromPlasticStatus(const FString& InResult)
  MV 100% Content\ToMove_BP.uasset -> Content\Moved_BP.uasset
  LM 100% Content\ToMove2_BP.uasset -> Content\Moved2_BP.uasset
  */
-static void ParseFileStatusResult(const TArray<FString>& InFiles, const TArray<FString>& InResults, TArray<FPlasticSourceControlState>& OutStates, int32& OutChangeset, FString& OutBranchName)
+static void ParseFileStatusResult(TArray<FString>&& InFiles, const TArray<FString>& InResults, TArray<FPlasticSourceControlState>& OutStates, int32& OutChangeset, FString& OutBranchName)
 {
 	const FPlasticSourceControlModule& PlasticSourceControl = FModuleManager::GetModuleChecked<FPlasticSourceControlModule>("PlasticSourceControl");
 	const FString& WorkingDirectory = PlasticSourceControl.GetProvider().GetPathToWorkspaceRoot();
 
-	// Parse the first two lines with Changeset number and Branch name
+	// Parse the first two lines with Changeset number and Branch name (the second being requested only once at init)
 	FString RepositoryName, ServerUrl;
 	ParseWorkspaceInformation(InResults, OutChangeset, RepositoryName, ServerUrl, OutBranchName);
 
 	// Iterate on each file explicitly listed in the command
-	for (const FString& File : InFiles)
+	for (FString& InFile : InFiles)
 	{
-		FPlasticSourceControlState FileState(File);
+		FPlasticSourceControlState FileState(MoveTemp(InFile));
+		const FString& File = FileState.LocalFilename;
 
 		// Search the file in the list of status
 		// NOTE: in case of rename by editor, there are two results: checked-out AND renamed
@@ -742,47 +739,49 @@ static void ParseFileStatusResult(const TArray<FString>& InFiles, const TArray<F
 		}
 		FileState.TimeStamp.Now();
 
-		// @todo: temporary debug log (only for the first few files)
-		if (OutStates.Num() < 20) UE_LOG(LogSourceControl, Log, TEXT("%s = %d:%s"), *File, static_cast<uint32>(FileState.WorkspaceState), FileState.ToString());
+		// debug log (only for the first few files)
+		if (OutStates.Num() < 20)
+		{
+			UE_LOG(LogSourceControl, Verbose, TEXT("%s = %d:%s"), *File, static_cast<uint32>(FileState.WorkspaceState), FileState.ToString());
+		}
 
 		OutStates.Add(MoveTemp(FileState));
 	}
-	// @todo: temporary debug log (if too many files)
-	if (OutStates.Num() > 20) UE_LOG(LogSourceControl, Log, TEXT("[...] %d more files"), OutStates.Num() - 20);
+	// debug log (if too many files)
+	if (OutStates.Num() > 20)
+	{
+		UE_LOG(LogSourceControl, Verbose, TEXT("[...] %d more files"), OutStates.Num() - 20);
+	}
 }
 
 /**
- * @brief Parse the array of strings results of a 'cm status --noheaders --controlledchanged --changed --localdeleted --localmoved --private --ignored' command
- *
- * Called in case of a "directory status" (no file listed in the command) ONLY to detect Removed/Deleted files !
+ * @brief Detect Deleted files in case of a "whole directory status" (no file listed in the command)
+ * 
+ * Parse the array of strings results of a 'cm status --noheaders --all --ignored' command
  *
  * @param[in]	InResults	Lines of results from the "status" command
  * @param[out]	OutStates	States of files for witch the status has been gathered
  *
  * @see #ParseFileStatusResult() above for an example of a cm status results
 */
-static void ParseDirectoryStatusResult(const TArray<FString>& InResults, TArray<FPlasticSourceControlState>& OutStates)
+static void ParseDirectoryStatusResultForDeleted(const TArray<FString>& InResults, TArray<FPlasticSourceControlState>& OutStates)
 {
 	const FPlasticSourceControlModule& PlasticSourceControl = FModuleManager::GetModuleChecked<FPlasticSourceControlModule>("PlasticSourceControl");
-	const FString& PathToPlasticBinary = PlasticSourceControl.AccessSettings().GetBinaryPath();
 	const FString& WorkingDirectory = PlasticSourceControl.GetProvider().GetPathToWorkspaceRoot();
 
 	// Iterate on each line of result of the status command
-	// NOTE: in case of rename by editor, there are two results: checkouted AND renamed
-	// => we want to get the second one, witch is always the rename, so we just iterate and the second state will overwrite the first one
 	for (const FString& Result : InResults)
 	{
-		const FString RelativeFilename = FilenameFromPlasticStatus(Result);
-		const FString File = FPaths::ConvertRelativePathToFull(WorkingDirectory, RelativeFilename);
 		const EWorkspaceState::Type WorkspaceState = StateFromPlasticStatus(Result);
 		if ((EWorkspaceState::Deleted == WorkspaceState) || (EWorkspaceState::LocallyDeleted == WorkspaceState))
 		{
-			FPlasticSourceControlState FileState(File);
+			FString RelativeFilename = FilenameFromPlasticStatus(Result);
+			FString AbsoluteFilename = FPaths::ConvertRelativePathToFull(WorkingDirectory, MoveTemp(RelativeFilename));
+			FPlasticSourceControlState FileState(MoveTemp(AbsoluteFilename));
 			FileState.WorkspaceState = WorkspaceState;
 			FileState.TimeStamp.Now();
 
-			// @todo: temporary debug log
-			UE_LOG(LogSourceControl, Log, TEXT("%s = %d:%s"), *File, static_cast<uint32>(FileState.WorkspaceState), FileState.ToString());
+			UE_LOG(LogSourceControl, Verbose, TEXT("%s = %d:%s"), *FileState.LocalFilename, static_cast<uint32>(FileState.WorkspaceState), FileState.ToString());
 
 			OutStates.Add(MoveTemp(FileState));
 		}
@@ -808,74 +807,80 @@ public:
 /**
  * @brief Run a "status" command for a directory to get workspace file states
  *
+ *  ie. Changed, CheckedOut, Copied, Replaced, Added, Private, Ignored, Deleted, LocallyDeleted, Moved, LocallyMoved
+ *
  *  It is either a command for a whole directory (ie. "Content/", in case of "Submit to Source Control"),
  * or for one or more files all on a same directory (by design, since we group files by directory in RunUpdateStatus())
  *
+ * @param[in]	InDir				The path to the common directory of all the files listed after.
  * @param[in]	InFiles				List of files in a directory, or the path to the directory itself (never empty).
  * @param[out]	OutErrorMessages	Error messages from the "status" command
  * @param[out]	OutStates			States of files for witch the status has been gathered (distinct than InFiles in case of a "directory status")
  * @param[out]	OutChangeset		The current Changeset Number
  * @param[out]	OutBranchName		Name of the current checked-out branch
  */
-static bool RunStatus(const TArray<FString>& InFiles, const EConcurrency::Type InConcurrency, TArray<FString>& OutErrorMessages, TArray<FPlasticSourceControlState>& OutStates, int32& OutChangeset, FString& OutBranchName)
+static bool RunStatus(const FString& InDir, TArray<FString>&& InFiles, const EConcurrency::Type InConcurrency, TArray<FString>& OutErrorMessages, TArray<FPlasticSourceControlState>& OutStates, int32& OutChangeset, FString& OutBranchName)
 {
-	ensure(0 < InFiles.Num());
+	check(InFiles.Num() > 0);
 
 	TArray<FString> Parameters;
 
-	// Command-line format output changed with version 8.0.16.3000, plugin will crash unless we demand the old format.
-	// For more details, see https://www.plasticscm.com/download/releasenotes/oldernotes?release=8.0.16.3000
-	if (bIsNewVersion80163000) {
+	// Command-line format output changed with version 8.0.16.3000, see https://www.plasticscm.com/download/releasenotes/8.0.16.3000
+	if (bIsNewVersion80163000)
+	{
 		Parameters.Add(TEXT("--compact"));
 	}
-	Parameters.Add(TEXT("--wkconfig")); // Branch name. TODO Deprecated in 8.0.16.3000 https://www.plasticscm.com/download/releasenotes/oldernotes?release=8.0.16.3000 => by default in xml format
 	Parameters.Add(TEXT("--noheaders"));
-	Parameters.Add(TEXT("--controlledchanged --changed --localdeleted --localmoved --private --ignored"));
-	// "cm status" only operate on one path (file or folder) at a time, so use one folder path for multiple files in a directory
-	const FString Path = FPaths::GetPath(*InFiles[0]);
+	Parameters.Add(TEXT("--all"));
+	Parameters.Add(TEXT("--ignored"));
+	// "cm status" only operate on one path (file or directory) at a time, so use one common path for multiple files in a directory
 	TArray<FString> OnePath;
 	// Only one file: optim very useful for the .uproject file at the root to avoid parsing the whole repository
-	// (does not work if file does not exist anymore)
-	if ((1 == InFiles.Num()) && (FPaths::FileExists(InFiles[0])))
+	// (but doesn't work if the file is deleted)
+	const bool bSingleFile = (InFiles.Num() == 1) && (FPaths::FileExists(InFiles[0]));
+	if (bSingleFile)
 	{
 		OnePath.Add(InFiles[0]);
 	}
 	else
 	{
-		OnePath.Add(Path);
+		OnePath.Add(InDir);
 	}
 	TArray<FString> Results;
 	TArray<FString> ErrorMessages;
 	const bool bResult = RunCommand(TEXT("status"), Parameters, OnePath, InConcurrency, Results, ErrorMessages);
-	// Normalize paths in the result (convert all '\' to '/')
-	for (int32 IdxResult = 0; IdxResult < Results.Num(); IdxResult++)
-	{
-		FPaths::NormalizeFilename(Results[IdxResult]);
-	}
-	OutErrorMessages.Append(ErrorMessages);
+	OutErrorMessages.Append(MoveTemp(ErrorMessages));
 	if (bResult)
 	{
-		if (1 == InFiles.Num() && FPaths::DirectoryExists(InFiles[0]))
+		// Normalize paths in the result (convert all '\' to '/')
+		for (FString& Result : Results)
+		{
+			FPaths::NormalizeFilename(Result);
+		}
+
+		const bool bWholeDirectory = (InFiles.Num() == 1) && (InFiles[0] == InDir);
+		if (bWholeDirectory)
 		{
 			// 1) Special case for "status" of a directory: requires a specific parse logic.
 			//   (this is triggered by the "Submit to Source Control" top menu button)
 			// Find recursively all files in the directory: this enable getting the list of "Controlled" (unchanged) assets
 			FFileVisitor FileVisitor;
-			FPlatformFileManager::Get().GetPlatformFile().IterateDirectoryRecursively(*Path, FileVisitor);
-			// @todo: temporary debug log
-			UE_LOG(LogSourceControl, Log, TEXT("RunStatus(%s): 1) special case for status of a directory containing %d file(s) (%s)"), *InFiles[0], FileVisitor.Files.Num(), *Path);
-			ParseFileStatusResult(FileVisitor.Files, Results, OutStates, OutChangeset, OutBranchName);
+			FPlatformFileManager::Get().GetPlatformFile().IterateDirectoryRecursively(*InDir, FileVisitor);
+			UE_LOG(LogSourceControl, Verbose, TEXT("RunStatus(%s): 1) special case for status of a directory containing %d file(s)"), *InDir, FileVisitor.Files.Num());
+			ParseFileStatusResult(MoveTemp(FileVisitor.Files), Results, OutStates, OutChangeset, OutBranchName);
 			// The above cannot detect assets removed / locally deleted since there is no file left to enumerate (either by the Content Browser or by File Manager)
 			// => so we also parse the status results to explicitly look for Removed/Deleted assets
-			Results.RemoveAt(0, 2); // Before that, remove the first two line Changeset, and BranchName
-			ParseDirectoryStatusResult(Results, OutStates);
+			if (Results.Num() > 0)
+			{
+				Results.RemoveAt(0, 1);// Before that, remove the first line (Workspace/Changeset info)
+			}
+			ParseDirectoryStatusResultForDeleted(Results, OutStates);
 		}
 		else
 		{
 			// 2) General case for one or more files in the same directory.
-			// @todo: temporary debug log
-			UE_LOG(LogSourceControl, Log, TEXT("RunStatus(%s...): 2) general case for %d file(s) in a directory (%s)"), *InFiles[0], InFiles.Num(), *Path);
-			ParseFileStatusResult(InFiles, Results, OutStates, OutChangeset, OutBranchName);
+			UE_LOG(LogSourceControl, Verbose, TEXT("RunStatus(%s...): 2) general case for %d file(s) in a directory (%s)"), *InFiles[0], InFiles.Num(), *InDir);
+			ParseFileStatusResult(MoveTemp(InFiles), Results, OutStates, OutChangeset, OutBranchName);
 		}
 	}
 
@@ -929,12 +934,14 @@ static void ParseFileinfoResults(const TArray<FString>& InResults, TArray<FPlast
 	const FPlasticSourceControlModule& PlasticSourceControl = FModuleManager::GetModuleChecked<FPlasticSourceControlModule>("PlasticSourceControl");
 	const FPlasticSourceControlProvider& Provider = PlasticSourceControl.GetProvider();
 
+	ensureMsgf(InResults.Num() == InOutStates.Num(), TEXT("The fileinfo command should gives the same number of infos as the status command"));
+
 	// Iterate on all files and all status of the result (assuming same number of line of results than number of file states)
 	for (int32 IdxResult = 0; IdxResult < InResults.Num(); IdxResult++)
 	{
 		const FString& Fileinfo = InResults[IdxResult];
 		FPlasticSourceControlState& FileState = InOutStates[IdxResult];
-		const FString& File = FileState.GetFilename();
+		const FString& File = FileState.LocalFilename;
 		FPlasticFileinfoParser FileinfoParser(Fileinfo);
 
 		FileState.LocalRevisionChangeset = FileinfoParser.RevisionChangeset;
@@ -943,61 +950,86 @@ static void ParseFileinfoResults(const TArray<FString>& InResults, TArray<FPlast
 		FileState.LockedBy = MoveTemp(FileinfoParser.LockedBy);
 		FileState.LockedWhere = MoveTemp(FileinfoParser.LockedWhere);
 
-		// If a file is locked by someone but not checked-out locally (or moved/renamed) this means it is locked by someone else or elsewhere
-		if ((FileState.WorkspaceState != EWorkspaceState::CheckedOut) && (FileState.WorkspaceState != EWorkspaceState::Moved) && (0 < FileState.LockedBy.Len())) 
+		// If a file is locked but not checked-out locally (or moved/renamed) this means it is locked by someone else or elsewhere
+		if ((FileState.WorkspaceState != EWorkspaceState::CheckedOut) && (FileState.WorkspaceState != EWorkspaceState::Moved) && !FileState.LockedBy.IsEmpty()) 
 		{
-			// @todo: temporary debug log
-			UE_LOG(LogSourceControl, Log, TEXT("LockedByOther(%s) by '%s!=%s' (or %s!=%s)"), *File, *FileState.LockedBy, *Provider.GetUserName(), *FileState.LockedWhere, *Provider.GetWorkspaceName());
+			UE_LOG(LogSourceControl, Verbose, TEXT("LockedByOther(%s) by '%s!=%s' (or %s!=%s)"), *File, *FileState.LockedBy, *Provider.GetUserName(), *FileState.LockedWhere, *Provider.GetWorkspaceName());
 			FileState.WorkspaceState = EWorkspaceState::LockedByOther;
 		}
 
-		// @todo: temporary debug log (only for the first few files)
-		if (IdxResult < 20) UE_LOG(LogSourceControl, Log, TEXT("%s: %d;%d %s by '%s' (%s)"), *File, FileState.LocalRevisionChangeset, FileState.DepotRevisionChangeset, *FileState.RepSpec, *FileState.LockedBy, *FileState.LockedWhere);
+		// debug log (only for the first few files)
+		if (IdxResult < 20)
+		{
+			UE_LOG(LogSourceControl, Verbose, TEXT("%s: %d;%d %s by '%s' (%s)"), *File, FileState.LocalRevisionChangeset, FileState.DepotRevisionChangeset, *FileState.RepSpec, *FileState.LockedBy, *FileState.LockedWhere);
+		}
+	}
+	// debug log (if too many files)
+	if (InResults.Num() > 20)
+	{
+		UE_LOG(LogSourceControl, Verbose, TEXT("[...] %d more files"), InResults.Num() - 20);
 	}
-	// @todo: temporary debug log (if too many files)
-	if (InResults.Num() > 20) UE_LOG(LogSourceControl, Log, TEXT("[...] %d more files"), InResults.Num() - 20);
 }
 
 /**
- * @brief Run a Plastic "fileinfo" command to update status of given files.
+ * @brief Run a "fileinfo" command to update complementary status information of given files.
  *
- * @param[in]		InForceFileinfo		Also force execute the fileinfo command required to do get RepSpec of xlinks when getting history (or for diffs)
+ * ie RevisionChangeset, RevisionHeadChangeset, RepSpec, LockedBy, LockedWhere
+ *
+ * @param[in]		bInWholeDirectory	If executed on a whole directory (typically Content/) for a "Submit Content" operation, optimize fileinfo more agressively
+ * @param			bInUpdateHistory	If getting the history of files, force execute the fileinfo command required to get RepSpec of XLinks (history view or visual diff)
  * @param[in]		InConcurrency		Is the command running in the background, or blocking the main thread
  * @param[out]		OutErrorMessages	Error messages from the "fileinfo" command
  * @param[in,out]	InOutStates			List of file states in the directory, gathered by the "status" command, completed by results of the "fileinfo" command
  */
-static bool RunFileinfo(const bool InForceFileinfo, const EConcurrency::Type InConcurrency, TArray<FString>& OutErrorMessages, TArray<FPlasticSourceControlState>& InOutStates)
+static bool RunFileinfo(const bool bInWholeDirectory, const bool bInUpdateHistory, const EConcurrency::Type InConcurrency, TArray<FString>& OutErrorMessages, TArray<FPlasticSourceControlState>& InOutStates)
 {
 	bool bResult = true;
-	TArray<FString> Files;
-	for (const auto& State : InOutStates)
-	{
-		// Optimize by not issuing "fileinfo" commands on "Added"/"Deleted"/"NotControled"/"Ignored" but also "CheckedOut" and "Moved" files.
-		// This can greatly reduce the time needed to do some basic operation like "Add to source control" when using a distant server or the Plastic Cloud.
-		// this can't work with xlink file when we want to update the history
-		// we need to know that we are running a fileinfo command to get the history, that's the role of InForceFileinfo
-		if (	(InForceFileinfo)
-			||	(State.WorkspaceState == EWorkspaceState::Controlled)
+	TArray<FString> SelectedFiles;
+
+	TArray<FPlasticSourceControlState> SelectedStates;
+	TArray<FPlasticSourceControlState> OptimizedStates;
+	for (FPlasticSourceControlState& State : InOutStates)
+	{
+		// 1) Issue a "fileinfo" command for controled files (to know if they are up to date and can be checked-out or checked-in)
+		// but only if controlled unchanged, or locally changed / locally deleted,
+		// optimizing for files that are CheckedOut/Added/Deleted/Moved/Copied/Replaced/NotControled/Ignored/Private/Unknown
+		// (since there is no point to check if they are up to date in these cases; they are already checkedout or not controlld).
+		// This greatly reduce the time needed to do some operations like "Add" or "Move/Rename/Copy" when there is some latency with the server (eg cloud).
+		//
+		// 2) bInWholeDirectory: In the case of a "whole directory status" triggered by the "Submit Content" operation,
+		// don't even issue a "fileinfo" command for unchanged Controlled files since they won't be considered them for submit.
+		// This greatly reduce the time needed to open the Submit window.
+		// 
+		// 3) bInUpdateHistory: When the plugin needs to update the history of files, it needs to know if it's on a XLink,
+		// so the fileinfo command is required here to get the RepSpec
+		if (bInUpdateHistory
+			|| ((State.WorkspaceState == EWorkspaceState::Controlled) && !bInWholeDirectory)
 			||	(State.WorkspaceState == EWorkspaceState::Changed)
-			||	(State.WorkspaceState == EWorkspaceState::Replaced)
-			||	(State.WorkspaceState == EWorkspaceState::Conflicted)
-		//	||	(State.WorkspaceState == EWorkspaceState::LockedByOther) // we do not have this info at this stage, cf. ParseFileinfoResults()
+			||	(State.WorkspaceState == EWorkspaceState::LocallyDeleted)
 			)
 		{
-			Files.Add(State.GetFilename());
+			SelectedFiles.Add(State.LocalFilename);
+			SelectedStates.Add(MoveTemp(State));
+		}
+		else
+		{
+			OptimizedStates.Add(MoveTemp(State));
 		}
 	}
-	if (Files.Num() > 0)
+	InOutStates = MoveTemp(OptimizedStates);
+
+	if (SelectedStates.Num())
 	{
 		TArray<FString> Results;
 		TArray<FString> ErrorMessages;
 		TArray<FString> Parameters;
 		Parameters.Add(TEXT("--format=\"{RevisionChangeset};{RevisionHeadChangeset};{RepSpec};{LockedBy};{LockedWhere}\""));
-		bResult = RunCommand(TEXT("fileinfo"), Parameters, Files, InConcurrency, Results, ErrorMessages);
-		OutErrorMessages.Append(ErrorMessages);
+		bResult = RunCommand(TEXT("fileinfo"), Parameters, SelectedFiles, InConcurrency, Results, ErrorMessages);
+		OutErrorMessages.Append(MoveTemp(ErrorMessages));
 		if (bResult)
 		{
-			ParseFileinfoResults(Results, InOutStates);
+			ParseFileinfoResults(Results, SelectedStates);
+			InOutStates.Append(MoveTemp(SelectedStates));
 		}
 	}
 
@@ -1054,8 +1086,7 @@ bool RunCheckMergeStatus(const TArray<FString>& InFiles, TArray<FString>& OutErr
 		FString MergeProgressContent;
 		if (FFileHelper::LoadFileToString(MergeProgressContent, *MergeProgressFilename))
 		{
-			// @todo: temporary debug logs
-			UE_LOG(LogSourceControl, Log, TEXT("RunCheckMergeStatus: %s:\n%s"), *MergeProgressFilename, *MergeProgressContent);
+			UE_LOG(LogSourceControl, Verbose, TEXT("RunCheckMergeStatus: %s:\n%s"), *MergeProgressFilename, *MergeProgressContent);
 			// Content is in one line, looking like the following:
 			// Target: mount:56e62dd7-241f-41e9-8c6b-dd4ca4513e62#/#UEMergeTest@localhost:8087 merged from: Merge 4
 			// Target: mount:56e62dd7-241f-41e9-8c6b-dd4ca4513e62#/#UEMergeTest@localhost:8087 merged from: Cherrypicking 3
@@ -1101,7 +1132,7 @@ bool RunCheckMergeStatus(const TArray<FString>& InFiles, TArray<FString>& OutErr
 					Parameters.Add(TEXT("--machinereadable"));
 					// call 'cm merge cs:xxx --machinereadable' (only dry-run, whithout the --merge parameter)
 					bResult = RunCommand(TEXT("merge"), Parameters, TArray<FString>(), EConcurrency::Synchronous, Results, ErrorMessages);
-					OutErrorMessages.Append(ErrorMessages);
+					OutErrorMessages.Append(MoveTemp(ErrorMessages));
 					// Parse the result, one line for each conflicted files:
 					for (const FString& Result : Results)
 					{
@@ -1112,8 +1143,7 @@ bool RunCheckMergeStatus(const TArray<FString>& InFiles, TArray<FString>& OutErr
 							UE_LOG(LogSourceControl, Log, TEXT("State.LocalFilename: '%s'"), *State.LocalFilename);
 							if (State.LocalFilename.EndsWith(MergeConflict.Filename, ESearchCase::CaseSensitive))
 							{
-								// @todo: temporary debug log
-								UE_LOG(LogSourceControl, Log, TEXT("MergeConflict '%s' found Base cs:%d From cs:%d"), *MergeConflict.Filename, MergeConflict.BaseChangeset, MergeConflict.SourceChangeset);
+								UE_LOG(LogSourceControl, Verbose, TEXT("MergeConflict '%s' found Base cs:%d From cs:%d"), *MergeConflict.Filename, MergeConflict.BaseChangeset, MergeConflict.SourceChangeset);
 								State.WorkspaceState = EWorkspaceState::Conflicted;
 								State.PendingMergeFilename = MergeConflict.Filename;
 								State.PendingMergeBaseChangeset = MergeConflict.BaseChangeset;
@@ -1131,38 +1161,111 @@ bool RunCheckMergeStatus(const TArray<FString>& InFiles, TArray<FString>& OutErr
 	return bResult;
 }
 
+FString FindCommonDirectory(const FString& InPath1, const FString& InPath2)
+{
+	const int32 MinLen = FMath::Min(InPath1.Len(), InPath2.Len());
+	int32 IndexAfterLastCommonSeparator = 0;
+	for (int32 Index = 0; Index < MinLen; Index++)
+	{
+		if (InPath1[Index] != InPath2[Index])
+			break;
+		if (InPath1[Index] == TEXT('/'))
+			IndexAfterLastCommonSeparator = Index + 1;
+	}
+	return InPath1.Left(IndexAfterLastCommonSeparator);
+}
+
+// Structure to group all files belonging to a root dir, storing their best/longest common directory
+struct FFilesInCommonDir
+{
+	// Best/longest common directory, slash terminated, based on FindCommonDirectory()
+	FString			CommonDir;
+	TArray<FString>	Files;
+};
+
 // Run a batch of Plastic "status" and "fileinfo" commands to update status of given files and directories.
-bool RunUpdateStatus(const TArray<FString>& InFiles, const bool InForceFileinfo, const EConcurrency::Type InConcurrency, TArray<FString>& OutErrorMessages, TArray<FPlasticSourceControlState>& OutStates, int32& OutChangeset, FString& OutBranchName)
+bool RunUpdateStatus(const TArray<FString>& InFiles, const bool bInUpdateHistory, const EConcurrency::Type InConcurrency, TArray<FString>& OutErrorMessages, TArray<FPlasticSourceControlState>& OutStates, int32& OutChangeset, FString& OutBranchName)
 {
 	bool bResults = true;
 
-	// The "status" command only operate on one directory at a time
-	// (whole tree recursively) not on different folders with no common root.
-	// But "Submit to Source Control" ask for the State of 3 different directories, Engine/Content, Project/Content and Project/Config,
-	// In a same way, a check-in can involve files from different subdirectories, and UpdateStatus is called for them all at once.
+	const FPlasticSourceControlModule& PlasticSourceControl = FModuleManager::GetModuleChecked<FPlasticSourceControlModule>("PlasticSourceControl");
+	const FString& WorkspaceRoot = PlasticSourceControl.GetProvider().GetPathToWorkspaceRoot();
+
+	// The "status" command only operate on one directory-tree at a time (whole tree recursively)
+	// not on different folders with no common root.
+	// But "Submit to Source Control" ask for the State of many different directories,
+	// from Project/Content and Project/Config, Engine/Content, Engine/Plugins/<...>/Content...
+
+	// In a similar way, a check-in can involve files from different subdirectories, and UpdateStatus is called for all of them at once.
+
+	static TArray<FString> RootDirs =
+	{
+		FPaths::ConvertRelativePathToFull(FPaths::ProjectContentDir()),
+		FPaths::ConvertRelativePathToFull(FPaths::ProjectConfigDir()),
+		FPaths::ConvertRelativePathToFull(FPaths::ProjectPluginsDir()),
+		FPaths::ConvertRelativePathToFull(FPaths::GameSourceDir()),
+		FPaths::ConvertRelativePathToFull(FPaths::EngineContentDir())
+	};
 
 	// 1) So here we group files by path (ie. by subdirectory)
-	TMap<FString, TArray<FString>> GroupOfFiles;
+	TMap<FString, FFilesInCommonDir> GroupOfFiles;
 	for (const FString& File : InFiles)
 	{
-		const FString Path = FPaths::GetPath(*File);
-		TArray<FString>* Group = GroupOfFiles.Find(Path);
-		if (Group != nullptr)
+		// Discard all file/paths that are not under the workspace root (typically excluding the Engine content)
+		if (!File.StartsWith(WorkspaceRoot))
 		{
-			Group->Add(File);
+			UE_LOG(LogSourceControl, Verbose, TEXT("%s is out of the Workspace"), *File);
+			continue;
 		}
-		else
+
+		bool bDirFound = false;
+		for (const auto& RootDir : RootDirs)
 		{
-			TArray<FString> NewGroup;
-			NewGroup.Add(File);
-			GroupOfFiles.Add(Path, NewGroup);
+			if (File.StartsWith(RootDir))
+			{
+				FFilesInCommonDir* ExistingGroup = GroupOfFiles.Find(RootDir);
+				if (ExistingGroup != nullptr)
+				{
+					// See if we have to update the CommonDir
+					if (!File.StartsWith(ExistingGroup->CommonDir))
+					{
+						// the file is not in the same path, we need to find their common dir
+						ExistingGroup->CommonDir = FindCommonDirectory(ExistingGroup->CommonDir, File);
+					}
+					ExistingGroup->Files.Add(File);
+				}
+				else
+				{
+					FString Path = FPaths::GetPath(File) + TEXT("/");
+					GroupOfFiles.Add(RootDir, { MoveTemp(Path), {File}});
+				}
+
+				bDirFound = true;
+				break;
+			}
+		}
+
+		// If the file isn't part of our root directories, we simply add its directory as a new group.
+		// It means that the group is dedicated to the directory, and as such its CommonDir is the directory itself.
+		// This should be an edge case (typically the uproject file) .
+		if (!bDirFound)
+		{
+			FString Path = FPaths::GetPath(File) + TEXT("/");
+			FFilesInCommonDir* ExistingGroup = GroupOfFiles.Find(Path);
+			if (ExistingGroup != nullptr)
+			{
+				ExistingGroup->Files.Add(File);
+			}
+			else
+			{
+				GroupOfFiles.Add(Path, { MoveTemp(Path), {File} });
+			}
 		}
 	}
 
-	// @todo: temporary debug log
 	if (InFiles.Num() > 0)
 	{
-		UE_LOG(LogSourceControl, Log, TEXT("RunUpdateStatus: %d file(s)/%d directory(ies) ('%s'...)"), InFiles.Num(), GroupOfFiles.Num(), *InFiles[0]);
+		UE_LOG(LogSourceControl, Verbose, TEXT("RunUpdateStatus: %d file(s)/%d directory(ies) ('%s'...)"), InFiles.Num(), GroupOfFiles.Num(), *InFiles[0]);
 	}
 	else
 	{
@@ -1170,17 +1273,21 @@ bool RunUpdateStatus(const TArray<FString>& InFiles, const bool InForceFileinfo,
 	}
 
 	// 2) then we can batch Plastic status operation by subdirectory
-	for (const auto& Files : GroupOfFiles)
+	for (auto& Group : GroupOfFiles)
 	{
+		const bool bWholeDirectory = ((Group.Value.Files.Num() == 1) && (Group.Value.CommonDir == Group.Value.Files[0]));	
+
 		// Run a "status" command on the directory to get workspace file states.
+		// (ie. Changed, CheckedOut, Copied, Replaced, Added, Private, Ignored, Deleted, LocallyDeleted, Moved, LocallyMoved)
 		TArray<FPlasticSourceControlState> States;
-		const bool bGroupOk = RunStatus(Files.Value, InConcurrency, OutErrorMessages, States, OutChangeset, OutBranchName);
+		const bool bGroupOk = RunStatus(Group.Value.CommonDir, MoveTemp(Group.Value.Files), InConcurrency, OutErrorMessages, States, OutChangeset, OutBranchName);
 		if (bGroupOk && (States.Num() > 0))
 		{
-			// Run a "fileinfo" command to update status of given files.
-			// In case of "directory status", there is no explicit file in the group (it contains only the directory) 
+			// Run a "fileinfo" command to update complementary status information of given files.
+			// (ie RevisionChangeset, RevisionHeadChangeset, RepSpec, LockedBy, LockedWhere)
+			// In case of "whole directory status", there is no explicit file in the group (it contains only the directory) 
 			// => work on the list of files discovered by RunStatus()
-			bResults &= RunFileinfo(InForceFileinfo, InConcurrency, OutErrorMessages, States);
+			bResults &= RunFileinfo(bWholeDirectory, bInUpdateHistory, InConcurrency, OutErrorMessages, States);
 		}
 		OutStates.Append(MoveTemp(States));
 	}
@@ -1200,14 +1307,13 @@ bool RunDumpToFile(const FString& InPathToPlasticBinary, const FString& InRevSpe
 	FString Errors;
 
 	// start with the Plastic command itself, then add revspec and temp filename to dump
-	FString FullCommand = TEXT("cat ");
+	FString FullCommand = TEXT("cat \"");
 	FullCommand += InRevSpec;
-	FullCommand += TEXT(" --raw --file=\"");
+	FullCommand += TEXT("\" --raw --file=\"");
 	FullCommand += InDumpFileName;
 	FullCommand += TEXT("\"");
 
-	// @todo: temporary debug logs
-	UE_LOG(LogSourceControl, Log, TEXT("RunDumpToFile: '%s %s'"), *InPathToPlasticBinary, *FullCommand);
+	UE_LOG(LogSourceControl, Verbose, TEXT("RunDumpToFile: '%s %s'"), *InPathToPlasticBinary, *FullCommand);
 	const bool bResult = FPlatformProcess::ExecProcess(*InPathToPlasticBinary, *FullCommand, &ReturnCode, &Results, &Errors);
 	UE_LOG(LogSourceControl, Log, TEXT("RunDumpToFile: ExecProcess ReturnCode=%d Results='%s'"), ReturnCode, *Results);
 	if (!bResult || !Errors.IsEmpty())
@@ -1218,235 +1324,246 @@ bool RunDumpToFile(const FString& InPathToPlasticBinary, const FString& InRevSpe
 	return bResult;
 }
 
-
 /**
- * Translate file actions from Plastic 'cm log' command to keywords used by the Editor UI.
- *
- * @see SHistoryRevisionListRowContent::GenerateWidgetForColumn(): "add", "edit", "delete", "branch" and "integrate"
+ * Parse results of the 'cm history --moveddeleted --xml --encoding="utf-8"' command.
+ * 
+ * Results of the history command looks like that:
+<RevisionHistoriesResult>
+  <RevisionHistories>
+	<RevisionHistory>
+	  <ItemName>C:/Workspace/UE4PlasticPluginDev/Content/FirstPersonBP/Blueprints/BP_TestsRenamed.uasset</ItemName>
+	  <Revisions>
+		<Revision>
+		  <RevisionSpec>C:/Workspace/UE4PlasticPluginDev/Content/FirstPersonBP/Blueprints/BP_TestsRenamed.uasset#cs:7</RevisionSpec>
+		  <Branch>/main</Branch>
+		  <CreationDate>2019-10-14T09:52:07+02:00</CreationDate>
+		  <RevisionType>bin</RevisionType>
+		  <ChangesetNumber>7</ChangesetNumber>
+		  <Owner>SRombauts</Owner>
+		  <Comment>New tests</Comment>
+		  <Repository>UE4PlasticPluginDev</Repository>
+		  <Server>localhost:8087</Server>
+		  <RepositorySpec>UE4PlasticPluginDev@localhost:8087</RepositorySpec>
+		</Revision>
+		...
+		<Revision>
+		  <RevisionSpec>C:/Workspace/UE4PlasticPluginDev/Content/FirstPersonBP/Blueprints/BP_TestsRenamed.uasset#cs:12</RevisionSpec>
+		  <Branch>Removed /Content/FirstPersonBP/Blueprints/BP_TestsRenamed.uasset</Branch>
+		  <CreationDate>2022-04-28T16:00:37+02:00</CreationDate>
+		  <RevisionType />
+		  <ChangesetNumber>12</ChangesetNumber>
+		  <Owner>sebastien.rombauts</Owner>
+		  <Comment />
+		  <Repository>UE4PlasticPluginDev</Repository>
+		  <Server>localhost:8087</Server>
+		  <RepositorySpec>UE4PlasticPluginDev@localhost:8087</RepositorySpec>
+		</Revision>
+
+	  </Revisions>
+	</RevisionHistory>
+	<RevisionHistory>
+	  <ItemName>C:/Workspace/UE4PlasticPluginDev/Content/FirstPersonBP/Blueprints/BP_YetAnother.uasset</ItemName>
+		...
+	</RevisionHistory>
+  </RevisionHistories>
+</RevisionHistoriesResult>
 */
-FString TranslateAction(const FString& InAction)
+static bool ParseHistoryResults(const bool bInUpdateHistory, const FXmlFile& InXmlResult, TArray<FPlasticSourceControlState>& InOutStates)
 {
-	if (InAction.Equals(TEXT("Added")))
-	{
-		return TEXT("add");
-	}
-	else if (InAction.Equals(TEXT("Moved")))
-	{
-		return TEXT("branch");
-	}
-	else if (InAction.Equals(TEXT("Deleted")))
-	{
-		return TEXT("delete");
-	}
-	else // if (InAction.Equals(TEXT("Changed")))
-	{
-		return TEXT("edit");
-	}
-}
+	const FPlasticSourceControlModule& PlasticSourceControl = FModuleManager::GetModuleChecked<FPlasticSourceControlModule>("PlasticSourceControl");
+	const FPlasticSourceControlProvider& Provider = PlasticSourceControl.GetProvider();
+	const FString RootRepSpec = FString::Printf(TEXT("%s@%s"), *Provider.GetRepositoryName(), *Provider.GetServerUrl());
 
-/**
- * Parse the array of strings results of a 'cm log --xml' command
- *
- * Example cm log results:
-<?xml version="1.0" encoding="utf-8"?>
-<LogList>
-  <Changeset>
-    <ObjId>989</ObjId>
-    <ChangesetId>2</ChangesetId>
-    <Branch>/main</Branch>
-    <Comment>Ignore Collections and Developers content</Comment>
-    <Owner>dev</Owner>
-    <GUID>a985c487-0f54-45c5-b0ef-9b87c4c3c3f9</GUID>
-    <Changes>
-      <Item>
-        <Branch>/main</Branch>
-        <RevNo>2</RevNo>
-        <Owner>dev</Owner>
-        <RevId>985</RevId>
-        <ParentRevId>282</ParentRevId>
-        <SrcCmPath>/ignore.conf</SrcCmPath>
-        <SrcParentItemId>2</SrcParentItemId>
-        <DstCmPath>/ignore.conf</DstCmPath>
-        <DstParentItemId>2</DstParentItemId>
-        <Date>2016-04-18T10:44:49.0000000+02:00</Date>
-        <Type>Changed</Type>
-      </Item>
-    </Changes>
-    <Date>2016-04-18T10:44:49.0000000+02:00</Date>
-  </Changeset>
-</LogList>
-*/
-static void ParseLogResults(const FXmlFile& InXmlResult, FPlasticSourceControlRevision& OutSourceControlRevision)
-{
-	static const FString LogList(TEXT("LogList"));
-	static const FString Changeset(TEXT("Changeset"));
-	static const FString Comment(TEXT("Comment"));
-	static const FString Date(TEXT("Date"));
+	static const FString RevisionHistoriesResult(TEXT("RevisionHistoriesResult"));
+	static const FString RevisionHistories(TEXT("RevisionHistories"));
+	static const FString RevisionHistory(TEXT("RevisionHistory"));
+	static const FString ItemName(TEXT("ItemName"));
+	static const FString Revisions(TEXT("Revisions"));
+	static const FString Revision(TEXT("Revision"));
+	static const FString Branch(TEXT("Branch"));
+	static const FString CreationDate(TEXT("CreationDate"));
+	static const FString RevisionType(TEXT("RevisionType"));
+	static const FString ChangesetNumber(TEXT("ChangesetNumber"));
 	static const FString Owner(TEXT("Owner"));
-	static const FString Changes(TEXT("Changes"));
-	static const FString Item(TEXT("Item"));
-	static const FString RevId(TEXT("RevId"));
-	static const FString ParentRevId(TEXT("ParentRevId"));
-	static const FString SrcCmPath(TEXT("SrcCmPath"));
-	static const FString DstCmPath(TEXT("DstCmPath"));
-	static const FString Type(TEXT("Type"));
+	static const FString Comment(TEXT("Comment"));
 
-	const FXmlNode* LogListNode = InXmlResult.GetRootNode();
-	if (LogListNode == nullptr || LogListNode->GetTag() != LogList)
+	const FXmlNode* RevisionHistoriesResultNode = InXmlResult.GetRootNode();
+	if (RevisionHistoriesResultNode == nullptr || RevisionHistoriesResultNode->GetTag() != RevisionHistoriesResult)
 	{
-		return;
+		return false;
 	}
 
-	const FXmlNode* ChangesetNode = LogListNode->FindChildNode(Changeset);
-	if (ChangesetNode == nullptr)
+	const FXmlNode* RevisionHistoriesNode = RevisionHistoriesResultNode->FindChildNode(RevisionHistories);
+	if (RevisionHistoriesNode == nullptr)
 	{
-		return;
+		return false;
 	}
 
-	const FXmlNode* CommentNode = ChangesetNode->FindChildNode(Comment);
-	if (CommentNode != nullptr)
+	const TArray<FXmlNode*>& RevisionHistoryNodes = RevisionHistoriesNode->GetChildrenNodes();
+	for (const FXmlNode* RevisionHistoryNode : RevisionHistoryNodes)
 	{
-		OutSourceControlRevision.Description = CommentNode->GetContent();
-	}
-	const FXmlNode* OwnerNode = ChangesetNode->FindChildNode(Owner);
-	if (OwnerNode != nullptr)
-	{
-		OutSourceControlRevision.UserName = OwnerNode->GetContent();
-	}
-	const FXmlNode* DateNode = ChangesetNode->FindChildNode(Date);
-	if (DateNode != nullptr)
-	{	//                           |--|
-		//    2016-04-18T10:44:49.0000000+02:00
-		// => 2016-04-18T10:44:49.000+02:00
-		const FString DateIso = DateNode->GetContent().LeftChop(10) + DateNode->GetContent().RightChop(27);
-		FDateTime::ParseIso8601(*DateIso, OutSourceControlRevision.Date);
-	}
+		const FXmlNode* ItemNameNode = RevisionHistoryNode->FindChildNode(ItemName);
+		if (ItemNameNode == nullptr)
+		{
+			continue;
+		}
 
-	const FXmlNode* ChangesNode = ChangesetNode->FindChildNode(Changes);
-	if (ChangesNode == nullptr)
-	{
-		return;
-	}
+		const FString Filename = ItemNameNode->GetContent();
+		FPlasticSourceControlState* InOutStatePtr = InOutStates.FindByPredicate(
+			[&Filename](const FPlasticSourceControlState& State) { return State.LocalFilename == Filename; }
+		);
+		if (InOutStatePtr == nullptr)
+		{
+			continue;
+		}
+		FPlasticSourceControlState& InOutState = *InOutStatePtr;
 
-	// Iterate on files to find the one we are tracking
-	for (const FXmlNode* ItemNode : ChangesNode->GetChildrenNodes())
-	{
-		int32 RevisionId = -1;
-		const FXmlNode* RevIdNode = ItemNode->FindChildNode(RevId);
-		if (RevIdNode != nullptr)
+		const FXmlNode* RevisionsNode = RevisionHistoryNode->FindChildNode(Revisions);
+		if (RevisionsNode == nullptr)
 		{
-			RevisionId = FCString::Atoi(*RevIdNode->GetContent());
+			continue;
 		}
-		// Is this about the file we are looking for?
-		if (RevisionId == OutSourceControlRevision.RevisionId)
+
+		const TArray<FXmlNode*>& RevisionNodes = RevisionsNode->GetChildrenNodes();
+		if (bInUpdateHistory)
 		{
-			const FXmlNode* DstCmPathNode = ItemNode->FindChildNode(DstCmPath);
-			if (DstCmPathNode != nullptr)
+			InOutState.History.Reserve(RevisionNodes.Num());
+		}
+
+		// parse history in reverse: needed to get most recent at the top (implied by the UI)
+		// Note: limit to last 100 changes, like Perforce
+		static const int32 MaxRevisions = 100;
+		const int32 MinIndex = FMath::Max(0, RevisionNodes.Num() - MaxRevisions);
+		for (int32 Index = RevisionNodes.Num() - 1; Index >= MinIndex; Index--)
+		{
+			if (const FXmlNode* RevisionNode = RevisionNodes[Index])
 			{
-				OutSourceControlRevision.Filename = DstCmPathNode->GetContent();
+				const TSharedRef<FPlasticSourceControlRevision, ESPMode::ThreadSafe> SourceControlRevision = MakeShared<FPlasticSourceControlRevision>();
+				SourceControlRevision->State = &InOutState;
+				SourceControlRevision->Filename = Filename;
+				SourceControlRevision->RevisionId = Index + 1;
 
-				const FXmlNode* SrcCmPathNode = ItemNode->FindChildNode(SrcCmPath);
-				const FXmlNode* ParentRevIdNode = ItemNode->FindChildNode(ParentRevId);
-				// Detect case of rename ("branch" in Perforce vocabulary)
-				if (ParentRevIdNode != nullptr && SrcCmPathNode != nullptr && !SrcCmPathNode->GetContent().Equals(DstCmPathNode->GetContent()))
+				if (const FXmlNode* RevisionTypeNode = RevisionNode->FindChildNode(RevisionType))
 				{
-					TSharedRef<FPlasticSourceControlRevision, ESPMode::ThreadSafe> MovedFromRevision = MakeShareable(new FPlasticSourceControlRevision);
-					MovedFromRevision->Filename = SrcCmPathNode->GetContent();
-					MovedFromRevision->RevisionId = FCString::Atoi(*ParentRevIdNode->GetContent());
-	
-					OutSourceControlRevision.BranchSource = MovedFromRevision;
+					if (!RevisionTypeNode->GetContent().IsEmpty())
+					{
+						if (Index == 0)
+						{
+							SourceControlRevision->Action = TEXT("add");
+						}
+						else
+						{
+							SourceControlRevision->Action = TEXT("edit");
+						}
+					}
+					else
+					{
+						SourceControlRevision->Action = TEXT("delete");
+					}
+				}
+
+				if (const FXmlNode* ChangesetNumberNode = RevisionNode->FindChildNode(ChangesetNumber))
+				{
+					const FString& Changeset = ChangesetNumberNode->GetContent();
+					SourceControlRevision->ChangesetNumber = FCString::Atoi(*Changeset); // Value now used in the Revision column and in the Asset Menu History
+
+					// Also append depot name to the revision, but only when it is different from the default one (ie for xlinks sub repository)
+					if (!InOutState.RepSpec.IsEmpty() && (InOutState.RepSpec != RootRepSpec))
+					{
+						TArray<FString> RepSpecs;
+						InOutState.RepSpec.ParseIntoArray(RepSpecs, TEXT("@"));
+						SourceControlRevision->Revision = FString::Printf(TEXT("cs:%s@%s"), *Changeset, *RepSpecs[0]);
+					}
+					else
+					{
+						SourceControlRevision->Revision = FString::Printf(TEXT("cs:%s"), *Changeset);
+					}
+				}
+				if (const FXmlNode* CommentNode = RevisionNode->FindChildNode(Comment))
+				{
+					SourceControlRevision->Description = CommentNode->GetContent();
+				}
+				if (const FXmlNode* OwnerNode = RevisionNode->FindChildNode(Owner))
+				{
+					SourceControlRevision->UserName = OwnerNode->GetContent();
+				}
+				if (const FXmlNode* DateNode = RevisionNode->FindChildNode(CreationDate))
+				{
+					FString DateIso = DateNode->GetContent();
+					const int len = DateIso.Len();
+					if (DateIso.Len() > 29)
+					{	//                           |--|
+						//    2016-04-18T10:44:49.0000000+02:00
+						// => 2016-04-18T10:44:49.000+02:00
+						DateIso = DateNode->GetContent().LeftChop(10) + DateNode->GetContent().RightChop(27);
+					}
+					FDateTime::ParseIso8601(*DateIso, SourceControlRevision->Date);
+				}
+				if (const FXmlNode* BranchNode = RevisionNode->FindChildNode(Branch))
+				{
+					SourceControlRevision->Branch = BranchNode->GetContent();
+				}
+
+				// Detect and skip more recent changesets on other branches (ie above the RevisionHeadChangeset)
+				if (SourceControlRevision->ChangesetNumber > InOutState.DepotRevisionChangeset)
+				{
+					InOutState.HeadBranch = SourceControlRevision->Branch;
+					InOutState.HeadAction = SourceControlRevision->Action;
+					InOutState.HeadChangeList = SourceControlRevision->ChangesetNumber;
+					InOutState.HeadUserName = SourceControlRevision->UserName;
+					InOutState.HeadModTime = SourceControlRevision->Date.ToUnixTimestamp();
+				}
+				else if (bInUpdateHistory)
+				{
+					InOutState.History.Add(SourceControlRevision);
+				}
+
+				if (!bInUpdateHistory)
+				{
+					break; // if not updating the history, just getting the head of the latest branch is enough
 				}
 			}
-			const FXmlNode* TypeNode = ItemNode->FindChildNode(Type);
-			if (TypeNode != nullptr)
-			{
-				OutSourceControlRevision.Action = TranslateAction(TypeNode->GetContent());
-			}
-			// 	Do not stop at first match, because in case of rename there are multiple log nodes: Changed+Moved (in this order)
 		}
 	}
+
+	return true;
 }
 
-// Run "cm log" on the changeset provided by the "history" command to get extra info about the change at a specific revision
-static bool RunLogCommand(const FString& InChangeset, const FString& InRepSpec, FPlasticSourceControlRevision& OutSourceControlRevision)
+// Run a Plastic "history" command and parse it's XML result.
+bool RunGetHistory(const bool bInUpdateHistory, TArray<FPlasticSourceControlState>& InOutStates, TArray<FString>& OutErrorMessages)
 {
-	const FString RepositorySpecification = FString::Printf(TEXT("cs:%s@%s"), *InChangeset, *InRepSpec);
-
+	bool bResult = true;
 	FString Results;
 	FString Errors;
 	TArray<FString> Parameters;
-	Parameters.Add(RepositorySpecification);
+	Parameters.Add(TEXT("--moveddeleted"));
 	Parameters.Add(TEXT("--xml"));
 	Parameters.Add(TEXT("--encoding=\"utf-8\""));
 
-	// Uses the raw RunCommandInternal() that does not split results in an array of strings, for XML parsing
-	bool bResult = RunCommandInternal(TEXT("log"), Parameters, TArray<FString>(), EConcurrency::Synchronous, Results, Errors);
-	if (bResult)
+	TArray<FString> Files;
+	Files.Reserve(InOutStates.Num());
+	for (const FPlasticSourceControlState& State : InOutStates)
 	{
-		FXmlFile XmlFile;
-		bResult = XmlFile.LoadFile(Results, EConstructMethod::ConstructFromBuffer);
-		if (bResult)
+		// When getting only the last revision, optimize out if DepotRevisionChangeset is invalid (ie "fileinfo" was optimized out, eg for checked-out files)
+		if (!bInUpdateHistory && State.DepotRevisionChangeset == ISourceControlState::INVALID_REVISION)
+			continue;
+
+		if (State.IsSourceControlled() && !State.IsAdded())
 		{
-			ParseLogResults(XmlFile, OutSourceControlRevision);
+			Files.Add(State.LocalFilename);
 		}
 	}
-	return bResult;
-}
-
-/**
- * Parse results of the 'cm history --format="{1};{6}"' command ("Changeset number" and "Revision id"),
- * then run "cm log" on each revision to get extra info about the change (description, date, filename, branch, action)
- * 
- * Results of the history command are with one changeset number and revision id by line, like that:
-14;176
-17;220
-18;223
-*/
-static bool ParseHistoryResults(const TArray<FString>& InResults, FPlasticSourceControlState& InOutState)
-{
-	bool bResult = true;
-
-	const FPlasticSourceControlModule& PlasticSourceControl = FModuleManager::GetModuleChecked<FPlasticSourceControlModule>("PlasticSourceControl");
-	const FPlasticSourceControlProvider& Provider = PlasticSourceControl.GetProvider();
-	const FString RootRepSpec = FString::Printf(TEXT("%s@%s"), *Provider.GetRepositoryName(), *Provider.GetServerUrl());
-
-	InOutState.History.Reserve(InResults.Num());
-
-	// parse history in reverse: needed to get most recent at the top (implied by the UI)
-	for (int32 Index = InResults.Num() - 1; Index >= 0; Index--)
+	if (Files.Num() > 0)
 	{
-		const FString& Result = InResults[Index];
+		bResult = RunCommandInternal(TEXT("history"), Parameters, Files, EConcurrency::Synchronous, Results, Errors);
+		OutErrorMessages.Add(MoveTemp(Errors));
 		if (bResult)
 		{
-			TArray<FString> Infos;
-			const int32 NbElmts = Result.ParseIntoArray(Infos, TEXT(";"));
-			if (NbElmts == 2)
+			FXmlFile XmlFile;
+			bResult = XmlFile.LoadFile(Results, EConstructMethod::ConstructFromBuffer);
+			if (bResult)
 			{
-				const TSharedRef<FPlasticSourceControlRevision, ESPMode::ThreadSafe> SourceControlRevision = MakeShareable(new FPlasticSourceControlRevision);
-				SourceControlRevision->State = &InOutState;
-				const FString& Changeset = Infos[0];
-				const FString& RevisionId = Infos[1];
-				SourceControlRevision->ChangesetNumber = FCString::Atoi(*Changeset); // Value now used in the Revision column and in the Asset Menu History
-				SourceControlRevision->RevisionId = FCString::Atoi(*RevisionId); // 
-				// Also append depot name to the revision, but only when it is different from the default one (ie for xlinks sub repository)
-				if (InOutState.RepSpec != RootRepSpec)
-				{
-					TArray<FString> RepSpecs;
-					InOutState.RepSpec.ParseIntoArray(RepSpecs, TEXT("@"));
-					SourceControlRevision->Revision = FString::Printf(TEXT("cs:%s@%s"), *Changeset, *RepSpecs[0]);
-				}
-				else
-				{
-					SourceControlRevision->Revision = FString::Printf(TEXT("cs:%s"), *Changeset);
-				}
-
-				// Run "cm log" on the changeset number
-				bResult = RunLogCommand(Changeset, InOutState.RepSpec, *SourceControlRevision);
-				InOutState.History.Add(SourceControlRevision);
-			}
-			else
-			{
-				bResult = false;
+				bResult = ParseHistoryResults(bInUpdateHistory, XmlFile, InOutStates);
 			}
 		}
 	}
@@ -1454,24 +1571,6 @@ static bool ParseHistoryResults(const TArray<FString>& InResults, FPlasticSource
 	return bResult;
 }
 
-// Run a Plastic "history" command and multiple "log" commands and parse them.
-bool RunGetHistory(const FString& InFile, TArray<FString>& OutErrorMessages, FPlasticSourceControlState& InOutState)
-{
-	TArray<FString> Results;
-	TArray<FString> Parameters;
-	Parameters.Add(TEXT("--format=\"{1};{6}\"")); // Get "Changeset number" and "Revision id" of each revision of the asset
-	TArray<FString> OneFile;
-	OneFile.Add(*InFile);
-
-	bool bResult = RunCommand(TEXT("history"), Parameters, OneFile, EConcurrency::Synchronous, Results, OutErrorMessages);
-	if (bResult)
-	{
-		bResult = ParseHistoryResults(Results, InOutState);
-	}
-
-	return bResult;
-}
-
 bool UpdateCachedStates(TArray<FPlasticSourceControlState>&& InStates)
 {
 	FPlasticSourceControlModule& PlasticSourceControl = FModuleManager::GetModuleChecked<FPlasticSourceControlModule>( "PlasticSourceControl" );
@@ -1500,7 +1599,7 @@ struct FRemoveRedundantErrors
 
 	bool operator()(const FString& String) const
 	{
-		if(String.Contains(Filter))
+		if (String.Contains(Filter))
 		{
 			return true;
 		}
@@ -1515,9 +1614,9 @@ struct FRemoveRedundantErrors
 void RemoveRedundantErrors(FPlasticSourceControlCommand& InCommand, const FString& InFilter)
 {
 	bool bFoundRedundantError = false;
-	for(const FString& ErrorMessage : InCommand.ErrorMessages)
+	for (const FString& ErrorMessage : InCommand.ErrorMessages)
 	{
-		if(ErrorMessage.Contains(InFilter, ESearchCase::CaseSensitive, ESearchDir::FromEnd))
+		if (ErrorMessage.Contains(InFilter, ESearchCase::CaseSensitive, ESearchDir::FromEnd))
 		{
 			InCommand.InfoMessages.Add(ErrorMessage);
 			bFoundRedundantError = true;
@@ -1527,10 +1626,22 @@ void RemoveRedundantErrors(FPlasticSourceControlCommand& InCommand, const FStrin
 	InCommand.ErrorMessages.RemoveAll( FRemoveRedundantErrors(InFilter) );
 
 	// if we have no error messages now, assume success!
-	if(bFoundRedundantError && InCommand.ErrorMessages.Num() == 0 && !InCommand.bCommandSuccessful)
+	if (bFoundRedundantError && InCommand.ErrorMessages.Num() == 0 && !InCommand.bCommandSuccessful)
 	{
 		InCommand.bCommandSuccessful = true;
 	}
 }
 
+void SwitchVerboseLogs(const bool bInEnable)
+{
+	if (bInEnable && LogSourceControl.GetVerbosity() < ELogVerbosity::Verbose)
+	{
+		LogSourceControl.SetVerbosity(ELogVerbosity::Verbose);
+	}
+	else if (!bInEnable && LogSourceControl.GetVerbosity() == ELogVerbosity::Verbose)
+	{
+		LogSourceControl.SetVerbosity(ELogVerbosity::Log);
+	}
 }
+
+}
\ No newline at end of file
diff --git a/Engine/Plugins/Developer/PlasticSourceControl/Source/PlasticSourceControl/Private/PlasticSourceControlUtils.h b/Engine/Plugins/Developer/PlasticSourceControl/Source/PlasticSourceControl/Private/PlasticSourceControlUtils.h
--- a/Engine/Plugins/Developer/PlasticSourceControl/Source/PlasticSourceControl/Private/PlasticSourceControlUtils.h
+++ b/Engine/Plugins/Developer/PlasticSourceControl/Source/PlasticSourceControl/Private/PlasticSourceControlUtils.h
@@ -96,12 +96,14 @@ bool GetWorkspaceInformation(int32& OutChangeset, FString& OutRepositoryName, FS
  * @returns true if the command succeeded and returned no errors
  */
 bool RunCommand(const FString& InCommand, const TArray<FString>& InParameters, const TArray<FString>& InFiles, const EConcurrency::Type InConcurrency, TArray<FString>& OutResults, TArray<FString>& OutErrorMessages);
+// Run a Plastic command - output is a string.
+bool RunCommandInternal(const FString& InCommand, const TArray<FString>& InParameters, const TArray<FString>& InFiles, const EConcurrency::Type InConcurrency, FString& OutResults, FString& OutErrors);
 
 /**
  * Run a Plastic "status" command and parse it.
  *
  * @param	InFiles				The files to be operated on
- * @param	InForceFileinfo		Also force execute the fileinfo command required to do get RepSpec of xlinks when getting history (or for diffs)
+ * @param	bInUpdateHistory	If getting the history of files, force execute the fileinfo command required to do get RepSpec of xlinks (history view or visual diff)
  * @param	InConcurrency		Is the command running in the background, or blocking the main thread
  * @param	OutErrorMessages	Any errors (from StdErr) as an array per-line
  * @param	OutStates			States of the files
@@ -109,7 +111,7 @@ bool RunCommand(const FString& InCommand, const TArray<FString>& InParameters, c
  * @param	OutBranchName		Name of the current checked-out branch
  * @returns true if the command succeeded and returned no errors
  */
-bool RunUpdateStatus(const TArray<FString>& InFiles, const bool InForceFileinfo, const EConcurrency::Type InConcurrency, TArray<FString>& OutErrorMessages, TArray<FPlasticSourceControlState>& OutStates, int32& OutChangeset, FString& OutBranchName);
+bool RunUpdateStatus(const TArray<FString>& InFiles, const bool bInUpdateHistory, const EConcurrency::Type InConcurrency, TArray<FString>& OutErrorMessages, TArray<FPlasticSourceControlState>& OutStates, int32& OutChangeset, FString& OutBranchName);
 
 /**
  * Run a Plastic "cat" command to dump the binary content of a revision into a file.
@@ -124,11 +126,11 @@ bool RunDumpToFile(const FString& InPathToPlasticBinary, const FString& InRevSpe
 /**
  * Run a Plastic "history" and "log" commands and parse it.
  *
- * @param	InFile				The file to be operated on
+ * @param	bInUpdateHistory	If getting the history of files, versus only checking the heads of branches to detect newer commits
+ * @param	InOutStates			The file states to update with the history
  * @param	OutErrorMessages	Any errors (from StdErr) as an array per-line
- * @param	InOutState			The status to update with the history of the file
  */
-bool RunGetHistory(const FString& InFile, TArray<FString>& OutErrorMessages, FPlasticSourceControlState& InOutState);
+bool RunGetHistory(const bool bInUpdateHistory, TArray<FPlasticSourceControlState>& InOutStates, TArray<FString>& OutErrorMessages);
 
 /**
  * Helper function for various commands to update cached states.
@@ -142,4 +144,17 @@ bool UpdateCachedStates(TArray<FPlasticSourceControlState>&& InStates);
  */
 void RemoveRedundantErrors(FPlasticSourceControlCommand& InCommand, const FString& InFilter);
 
+/**
+ * Change LogSourceControl verbosity level at startup and when toggled from the Plastic Source Control Settings
+ *
+ * Override to Verbose or back to Log, but only if the current log verbosity is not already set to VeryVerbose
+ */
+void SwitchVerboseLogs(const bool bInEnable);
+
+/**
+ * Find the best(longest) common directory between two paths, terminated by a slash, returning an empty string if none.
+ * Assumes that both input strings are already normalized paths, slash delimited, for performance reason.
+ */
+FString FindCommonDirectory(const FString& InPath1, const FString& InPath2);
+
 }
diff --git a/Engine/Plugins/Developer/PlasticSourceControl/Source/PlasticSourceControl/Private/SPlasticSourceControlSettings.cpp b/Engine/Plugins/Developer/PlasticSourceControl/Source/PlasticSourceControl/Private/SPlasticSourceControlSettings.cpp
--- a/Engine/Plugins/Developer/PlasticSourceControl/Source/PlasticSourceControl/Private/SPlasticSourceControlSettings.cpp
+++ b/Engine/Plugins/Developer/PlasticSourceControl/Source/PlasticSourceControl/Private/SPlasticSourceControlSettings.cpp
@@ -132,13 +132,13 @@ void SPlasticSourceControlSettings::Construct(const FArguments& InArgs)
 		[
 			SNew(SHorizontalBox)
 			.Visibility(this, &SPlasticSourceControlSettings::CanInitializePlasticWorkspace)
-			.ToolTipText(LOCTEXT("WorkspaceNotFound_Tooltip", "No Workspace found at the level or above the current Project"))
+			.ToolTipText(LOCTEXT("WorkspaceNotFound_Tooltip", "No Workspace found at the level or above the current Project. Use the form to create a new one."))
 			+SHorizontalBox::Slot()
 			.FillWidth(1.0f)
 			.HAlign(HAlign_Center)
 			[
 				SNew(STextBlock)
-				.Text(LOCTEXT("WorkspaceNotFound", "Current Project is not contained in a Plastic SCM Workspace. Fill the form below to initialize a new Workspace."))
+				.Text(LOCTEXT("WorkspaceNotFound", "Current Project is not in a Plastic SCM Workspace. Create a new one:"))
 				.Font(Font)
 			]
 		]
@@ -154,7 +154,7 @@ void SPlasticSourceControlSettings::Construct(const FArguments& InArgs)
 			.FillWidth(1.0f)
 			[
 				SNew(STextBlock)
-				.Text(LOCTEXT("WorkspaceRepositoryName", "Workspace and Repository Name"))
+				.Text(LOCTEXT("WorkspaceRepositoryName", "Workspace & Repository"))
 				.ToolTipText(LOCTEXT("WorkspaceRepositoryName_Tooltip", "Enter the Name of the new Workspace and Repository to create or use"))
 				.Font(Font)
 			]
@@ -189,7 +189,7 @@ void SPlasticSourceControlSettings::Construct(const FArguments& InArgs)
 		[
 			SNew(SHorizontalBox)
 			.Visibility(this, &SPlasticSourceControlSettings::CanInitializePlasticWorkspace)
-			.ToolTipText(LOCTEXT("ServerUrl_Tooltip", "Enter the Server URL in the form address:port (localhost:8087)"))
+			.ToolTipText(LOCTEXT("ServerUrl_Tooltip", "Enter the Server URL in the form address:port (eg. localhost:8087 or Organization@cloud)"))
 			+SHorizontalBox::Slot()
 			.FillWidth(1.0f)
 			[
@@ -224,7 +224,7 @@ void SPlasticSourceControlSettings::Construct(const FArguments& InArgs)
 				.OnCheckStateChanged(this, &SPlasticSourceControlSettings::OnCheckedCreateIgnoreFile)
 			]
 			+SHorizontalBox::Slot()
-			.FillWidth(2.9f)
+			.FillWidth(2.0f)
 			.VAlign(VAlign_Center)
 			[
 				SNew(STextBlock)
@@ -249,15 +249,15 @@ void SPlasticSourceControlSettings::Construct(const FArguments& InArgs)
 				.OnCheckStateChanged(this, &SPlasticSourceControlSettings::OnCheckedInitialCommit)
 			]
 			+SHorizontalBox::Slot()
-			.FillWidth(0.9f)
+			.FillWidth(0.6f)
 			.VAlign(VAlign_Center)
 			[
 				SNew(STextBlock)
-				.Text(LOCTEXT("InitialCommit", "Make the initial Plastic SCM Checkin"))
+				.Text(LOCTEXT("InitialCommit", "Initial Checkin"))
 				.Font(Font)
 			]
 			+SHorizontalBox::Slot()
-			.FillWidth(2.0f)
+			.FillWidth(1.4f)
 			.Padding(2.0f)
 			[
 				SNew(SMultiLineEditableTextBox)
@@ -267,20 +267,76 @@ void SPlasticSourceControlSettings::Construct(const FArguments& InArgs)
 				.Font(Font)
 			]
 		]
-		// Option to enable the Update Status operation at Editor Startup
+		// Option to run an Update Status operation at Editor Startup
 		+SVerticalBox::Slot()
 		.AutoHeight()
 		.Padding(2.0f)
 		.VAlign(VAlign_Center)
 		[
-			SNew(SCheckBox)
-			.ForegroundColor(FSlateColor::UseForeground())
+			SNew(SHorizontalBox)
 			.ToolTipText(LOCTEXT("UpdateStatusAtStartup_Tooltip", "Run an asynchronous Update Status at Editor startup (can be slow)."))
-			.IsChecked(SPlasticSourceControlSettings::IsUpdateStatusAtStartupChecked())
-			.OnCheckStateChanged(this, &SPlasticSourceControlSettings::OnCheckedUpdateStatusAtStartup)
+			+SHorizontalBox::Slot()
+			.FillWidth(0.1f)
+			[
+				SNew(SCheckBox)
+				.IsChecked(SPlasticSourceControlSettings::IsUpdateStatusAtStartupChecked())
+				.OnCheckStateChanged(this, &SPlasticSourceControlSettings::OnCheckedUpdateStatusAtStartup)
+			]
+			+SHorizontalBox::Slot()
+			.FillWidth(2.0f)
+			.VAlign(VAlign_Center)
 			[
 				SNew(STextBlock)
 				.Text(LOCTEXT("UpdateStatusAtStartup", "Update workspace Status at Editor startup"))
+				.Font(Font)
+			]
+		]
+		// Option to call History as part of Update Status operation to check for potential recent changes in other branches
+		+SVerticalBox::Slot()
+		.AutoHeight()
+		.Padding(2.0f)
+		.VAlign(VAlign_Center)
+		[
+			SNew(SHorizontalBox)
+			.ToolTipText(LOCTEXT("UpdateStatusOtherBranches_Tooltip", "Enable Update status to detect changesets on other branches (can be slow)."))
+			+SHorizontalBox::Slot()
+			.FillWidth(0.1f)
+			[
+				SNew(SCheckBox)
+				.IsChecked(SPlasticSourceControlSettings::IsUpdateStatusOtherBranchesChecked())
+				.OnCheckStateChanged(this, &SPlasticSourceControlSettings::OnCheckedUpdateStatusOtherBranches)
+			]
+			+SHorizontalBox::Slot()
+			.FillWidth(2.0f)
+			.VAlign(VAlign_Center)
+			[
+				SNew(STextBlock)
+				.Text(LOCTEXT("UpdateStatusOtherBranches", "Update status also check branch history."))
+				.Font(Font)
+			]
+		]
+		// Option to enable Source Control Verbose logs
+		+SVerticalBox::Slot()
+		.AutoHeight()
+		.Padding(2.0f)
+		.VAlign(VAlign_Center)
+		[
+			SNew(SHorizontalBox)
+			.ToolTipText(LOCTEXT("EnableVerboseLogs_Tooltip", "Override LogSourceControl default verbosity level to Verbose (except if already set to VeryVerbose)."))
+			+SHorizontalBox::Slot()
+			.FillWidth(0.1f)
+			[
+				SNew(SCheckBox)
+				.IsChecked(SPlasticSourceControlSettings::IsEnableVerboseLogsChecked())
+				.OnCheckStateChanged(this, &SPlasticSourceControlSettings::OnCheckedEnableVerboseLogs)
+			]
+			+SHorizontalBox::Slot()
+			.FillWidth(2.0f)
+			.VAlign(VAlign_Center)
+			[
+				SNew(STextBlock)
+				.Text(LOCTEXT("EnableVerboseLogs", "Enable Source Control Verbose logs"))
+				.Font(Font)
 			]
 		]
 		// Button to create a new Workspace
@@ -334,7 +390,7 @@ void SPlasticSourceControlSettings::OnBinaryPathTextCommited(const FText& InText
 {
 	FPlasticSourceControlModule& PlasticSourceControl = FModuleManager::LoadModuleChecked<FPlasticSourceControlModule>("PlasticSourceControl");
 	const bool bChanged = PlasticSourceControl.AccessSettings().SetBinaryPath(InText.ToString());
-	if(bChanged)
+	if (bChanged)
 	{
 		// Re-Check provided Plastic binary path for each change
 		PlasticSourceControl.GetProvider().CheckPlasticAvailability();
@@ -498,7 +554,7 @@ void SPlasticSourceControlSettings::LaunchCheckInOperation()
 	TSharedRef<FCheckIn, ESPMode::ThreadSafe> CheckInOperation = ISourceControlOperation::Create<FCheckIn>();
 	CheckInOperation->SetDescription(InitialCommitMessage);
 	FPlasticSourceControlModule& PlasticSourceControl = FModuleManager::LoadModuleChecked<FPlasticSourceControlModule>("PlasticSourceControl");
-	const TArray<FString> ProjectFiles = GetProjectFiles();
+	const TArray<FString> ProjectFiles = GetProjectFiles(); // Note: listing files and folders is only needed for the update status operation following the checkin to know on what to operate
 	ECommandResult::Type Result = PlasticSourceControl.GetProvider().Execute(CheckInOperation, ProjectFiles, EConcurrency::Asynchronous, FSourceControlOperationComplete::CreateSP(this, &SPlasticSourceControlSettings::OnSourceControlOperationComplete));
 	if (Result == ECommandResult::Succeeded)
 	{
@@ -573,8 +629,7 @@ void SPlasticSourceControlSettings::DisplaySuccessNotification(const FName& InOp
 	Info.bUseSuccessFailIcons = true;
 	Info.Image = FAppStyle::GetBrush(TEXT("NotificationList.SuccessImage"));
 	FSlateNotificationManager::Get().AddNotification(Info);
-	// @todo temporary debug log
-	UE_LOG(LogSourceControl, Log, TEXT("%s"), *NotificationText.ToString());
+	UE_LOG(LogSourceControl, Verbose, TEXT("%s"), *NotificationText.ToString());
 }
 
 // Display a temporary failure notification at the end of the operation
@@ -584,7 +639,6 @@ void SPlasticSourceControlSettings::DisplayFailureNotification(const FName& InOp
 	FNotificationInfo Info(NotificationText);
 	Info.ExpireDuration = 8.0f;
 	FSlateNotificationManager::Get().AddNotification(Info);
-	// @todo temporary debug log
 	UE_LOG(LogSourceControl, Error, TEXT("%s"), *NotificationText.ToString());
 }
 
@@ -621,15 +675,38 @@ void SPlasticSourceControlSettings::OnCheckedUpdateStatusAtStartup(ECheckBoxStat
 	PlasticSourceControl.AccessSettings().SaveSettings();
 }
 
-bool SPlasticSourceControlSettings::UpdateStatusAtStartup() const
+ECheckBoxState SPlasticSourceControlSettings::IsUpdateStatusAtStartupChecked() const
+{
+	const FPlasticSourceControlModule& PlasticSourceControl = FModuleManager::GetModuleChecked<FPlasticSourceControlModule>("PlasticSourceControl");
+	return PlasticSourceControl.AccessSettings().GetUpdateStatusAtStartup() ? ECheckBoxState::Checked : ECheckBoxState::Unchecked;
+}
+
+void SPlasticSourceControlSettings::OnCheckedUpdateStatusOtherBranches(ECheckBoxState NewCheckedState)
+{
+	FPlasticSourceControlModule& PlasticSourceControl = FModuleManager::GetModuleChecked<FPlasticSourceControlModule>("PlasticSourceControl");
+	PlasticSourceControl.AccessSettings().SetUpdateStatusOtherBranches(NewCheckedState == ECheckBoxState::Checked);
+	PlasticSourceControl.AccessSettings().SaveSettings();
+}
+
+ECheckBoxState SPlasticSourceControlSettings::IsUpdateStatusOtherBranchesChecked() const
 {
 	const FPlasticSourceControlModule& PlasticSourceControl = FModuleManager::GetModuleChecked<FPlasticSourceControlModule>("PlasticSourceControl");
-	return PlasticSourceControl.AccessSettings().UpdateStatusAtStartup();
+	return PlasticSourceControl.AccessSettings().GetUpdateStatusOtherBranches() ? ECheckBoxState::Checked : ECheckBoxState::Unchecked;
 }
 
-ECheckBoxState SPlasticSourceControlSettings::IsUpdateStatusAtStartupChecked() const
+void SPlasticSourceControlSettings::OnCheckedEnableVerboseLogs(ECheckBoxState NewCheckedState)
+{
+	FPlasticSourceControlModule& PlasticSourceControl = FModuleManager::GetModuleChecked<FPlasticSourceControlModule>("PlasticSourceControl");
+	PlasticSourceControl.AccessSettings().SetEnableVerboseLogs(NewCheckedState == ECheckBoxState::Checked);
+	PlasticSourceControl.AccessSettings().SaveSettings();
+
+	PlasticSourceControlUtils::SwitchVerboseLogs(NewCheckedState == ECheckBoxState::Checked);
+}
+
+ECheckBoxState SPlasticSourceControlSettings::IsEnableVerboseLogsChecked() const
 {
-	return (UpdateStatusAtStartup() ? ECheckBoxState::Checked : ECheckBoxState::Unchecked);
+	const FPlasticSourceControlModule& PlasticSourceControl = FModuleManager::GetModuleChecked<FPlasticSourceControlModule>("PlasticSourceControl");
+	return PlasticSourceControl.AccessSettings().GetEnableVerboseLogs() ? ECheckBoxState::Checked : ECheckBoxState::Unchecked;
 }
 
 /** Path to the "ignore.conf" file */
diff --git a/Engine/Plugins/Developer/PlasticSourceControl/Source/PlasticSourceControl/Private/SPlasticSourceControlSettings.h b/Engine/Plugins/Developer/PlasticSourceControl/Source/PlasticSourceControl/Private/SPlasticSourceControlSettings.h
--- a/Engine/Plugins/Developer/PlasticSourceControl/Source/PlasticSourceControl/Private/SPlasticSourceControlSettings.h
+++ b/Engine/Plugins/Developer/PlasticSourceControl/Source/PlasticSourceControl/Private/SPlasticSourceControlSettings.h
@@ -53,7 +53,12 @@ private:
 
 	void OnCheckedUpdateStatusAtStartup(ECheckBoxState NewCheckedState);
 	ECheckBoxState IsUpdateStatusAtStartupChecked() const;
-	bool UpdateStatusAtStartup() const;
+
+	void OnCheckedUpdateStatusOtherBranches(ECheckBoxState NewCheckedState);
+	ECheckBoxState IsUpdateStatusOtherBranchesChecked() const;
+
+	void OnCheckedEnableVerboseLogs(ECheckBoxState NewCheckedState);
+	ECheckBoxState IsEnableVerboseLogsChecked() const;
 
 	void OnCheckedInitialCommit(ECheckBoxState NewCheckedState);
 	bool bAutoInitialCommit;
