Update PlasticSCM Plugin from 1.5.0 to 1.5.1

#9195 (Contributed by [REDACTED_PERSON])

diff --git a/Engine/Plugins/Developer/PlasticSourceControl/PlasticSourceControl.uplugin b/Engine/Plugins/Developer/PlasticSourceControl/PlasticSourceControl.uplugin
--- a/Engine/Plugins/Developer/PlasticSourceControl/PlasticSourceControl.uplugin
+++ b/Engine/Plugins/Developer/PlasticSourceControl/PlasticSourceControl.uplugin
@@ -1,27 +1 @@
-{
-	"FileVersion": 3,
-	"Version": 50,
-	"VersionName": "1.5.0",
-	"FriendlyName": "Plastic SCM",
-	"Description": "Plastic source control management",
-	"Category": "Source Control",
-	"CreatedBy": "SRombauts for Codice Software.",
-	"CreatedByURL": "http://srombauts.github.com",
-	"DocsURL": "[REDACTED_PROFILE_URL]/UE4PlasticPlugin",
-	"MarketplaceURL": "",
-	"SupportURL": "https://www.plasticscm.com/support",
-	"EnabledByDefault": true,
-	"CanContainContent": false,
-	"IsBetaVersion": false,
-	"Installed": false,
-	"Modules": [
-		{
-			"Name": "PlasticSourceControl",
-			"Type": "Editor",
-			"LoadingPhase": "Default",
-			"PlatformAllowList": [
-				"Win64"
-			]
-		}
-	]
-}
\ No newline at end of file
+{	"FileVersion": 3,	"Version": 51,	"VersionName": "1.5.1",	"FriendlyName": "Plastic SCM",	"Description": "Plastic source control management",	"Category": "Source Control",	"CreatedBy": "SRombauts for Codice Software.",	"CreatedByURL": "http://srombauts.github.com",	"DocsURL": "[REDACTED_PROFILE_URL]/UE4PlasticPlugin",	"MarketplaceURL": "",	"SupportURL": "https://www.plasticscm.com/support",	"EnabledByDefault": true,	"CanContainContent": false,	"IsBetaVersion": false,	"Installed": false,	"Modules": [		{			"Name": "PlasticSourceControl",			"Type": "Editor",			"LoadingPhase": "Default",			"PlatformAllowList": [				"Win64"			]		}	]}
\ No newline at end of file
diff --git a/Engine/Plugins/Developer/PlasticSourceControl/Source/PlasticSourceControl/Private/PlasticSourceControlConsole.h b/Engine/Plugins/Developer/PlasticSourceControl/Source/PlasticSourceControl/Private/PlasticSourceControlConsole.h
--- a/Engine/Plugins/Developer/PlasticSourceControl/Source/PlasticSourceControl/Private/PlasticSourceControlConsole.h
+++ b/Engine/Plugins/Developer/PlasticSourceControl/Source/PlasticSourceControl/Private/PlasticSourceControlConsole.h
@@ -3,7 +3,8 @@
 #pragma once
 
 #include "CoreMinimal.h"
-#include "HAL/IConsoleManager.h"
+
+class FAutoConsoleCommand;
 
 /**
  * Editor only console commands.
diff --git a/Engine/Plugins/Developer/PlasticSourceControl/Source/PlasticSourceControl/Private/PlasticSourceControlOperations.cpp b/Engine/Plugins/Developer/PlasticSourceControl/Source/PlasticSourceControl/Private/PlasticSourceControlOperations.cpp
--- a/Engine/Plugins/Developer/PlasticSourceControl/Source/PlasticSourceControl/Private/PlasticSourceControlOperations.cpp
+++ b/Engine/Plugins/Developer/PlasticSourceControl/Source/PlasticSourceControl/Private/PlasticSourceControlOperations.cpp
@@ -181,6 +181,8 @@ bool FPlasticCheckInWorker::Execute(FPlasticSourceControlCommand& InCommand)
 	check(InCommand.Operation->GetName() == GetName());
 	TSharedRef<FCheckIn, ESPMode::ThreadSafe> Operation = StaticCastSharedRef<FCheckIn>(InCommand.Operation);
 
+	UE_LOG(LogSourceControl, Verbose, TEXT("CheckIn: %d file(s) Description: '%s'"), InCommand.Files.Num(), *Operation->GetDescription().ToString());
+
 	// make a temp file to place our commit message in
 	FScopedTempFile CommitMsgFile(Operation->GetDescription());
 	if (!CommitMsgFile.GetFilename().IsEmpty())
@@ -323,41 +325,50 @@ bool FPlasticRevertWorker::Execute(FPlasticSourceControlCommand& InCommand)
 	FPlasticSourceControlModule& PlasticSourceControl = FModuleManager::GetModuleChecked<FPlasticSourceControlModule>("PlasticSourceControl");
 	FPlasticSourceControlProvider& Provider = PlasticSourceControl.GetProvider();
 
-	InCommand.bCommandSuccessful = true;
+	TArray<FString> ChangedFiles;
+	TArray<FString> CheckedOutFiles;
 
 	for (const FString& File : InCommand.Files)
 	{
 		TSharedRef<FPlasticSourceControlState, ESPMode::ThreadSafe> State = Provider.GetStateInternal(File);
 
-		TArray<FString> Files;
-		Files.Add(State->LocalFilename);
-
 		if (EWorkspaceState::Changed == State->WorkspaceState)
 		{
-			// revert the changes of the given file in workspace
-			InCommand.bCommandSuccessful &= PlasticSourceControlUtils::RunCommand(TEXT("undochange"), TArray<FString>(), Files, InCommand.Concurrency, InCommand.InfoMessages, InCommand.ErrorMessages);
+			// only revert the changes of the given file in workspace
+			ChangedFiles.Add(State->LocalFilename);
 		}
 		else
 		{
+			CheckedOutFiles.Add(State->LocalFilename);
 			// in case of a Moved/Renamed, find the rename origin to revert both at once
 			if (EWorkspaceState::Moved == State->WorkspaceState)
 			{
-				Files.Add(State->MovedFrom);
+				CheckedOutFiles.Add(State->MovedFrom);
 
 				// Delete the redirector
 				IFileManager::Get().Delete(*State->MovedFrom);
 			}
+		}
+	}
 
-			// revert the checkout and any changes of the given file in workspace
-			// Detect special case for a partial checkout (CS:-1 in Gluon mode)!
-			if (-1 != InCommand.ChangesetNumber)
-			{
-				InCommand.bCommandSuccessful &= PlasticSourceControlUtils::RunCommand(TEXT("undocheckout"), TArray<FString>(), Files, InCommand.Concurrency, InCommand.InfoMessages, InCommand.ErrorMessages);
-			}
-			else
-			{
-				InCommand.bCommandSuccessful &= PlasticSourceControlUtils::RunCommand(TEXT("partial undocheckout"), TArray<FString>(), Files, InCommand.Concurrency, InCommand.InfoMessages, InCommand.ErrorMessages);
-			}
+	InCommand.bCommandSuccessful = true;
+
+	if (ChangedFiles.Num() > 0)
+	{
+		InCommand.bCommandSuccessful &= PlasticSourceControlUtils::RunCommand(TEXT("undochange"), TArray<FString>(), ChangedFiles, InCommand.Concurrency, InCommand.InfoMessages, InCommand.ErrorMessages);
+	}
+
+	if (CheckedOutFiles.Num() > 0)
+	{
+		// revert the checkout and any changes of the given file in workspace
+		// Detect special case for a partial checkout (CS:-1 in Gluon mode)!
+		if (-1 != InCommand.ChangesetNumber)
+		{
+			InCommand.bCommandSuccessful &= PlasticSourceControlUtils::RunCommand(TEXT("undocheckout"), TArray<FString>(), CheckedOutFiles, InCommand.Concurrency, InCommand.InfoMessages, InCommand.ErrorMessages);
+		}
+		else
+		{
+			InCommand.bCommandSuccessful &= PlasticSourceControlUtils::RunCommand(TEXT("partial undocheckout"), TArray<FString>(), CheckedOutFiles, InCommand.Concurrency, InCommand.InfoMessages, InCommand.ErrorMessages);
 		}
 	}
 
diff --git a/Engine/Plugins/Developer/PlasticSourceControl/Source/PlasticSourceControl/Private/PlasticSourceControlProvider.cpp b/Engine/Plugins/Developer/PlasticSourceControl/Source/PlasticSourceControl/Private/PlasticSourceControlProvider.cpp
--- a/Engine/Plugins/Developer/PlasticSourceControl/Source/PlasticSourceControl/Private/PlasticSourceControlProvider.cpp
+++ b/Engine/Plugins/Developer/PlasticSourceControl/Source/PlasticSourceControl/Private/PlasticSourceControlProvider.cpp
@@ -47,7 +47,21 @@ void FPlasticSourceControlProvider::Init(bool bForceConnection)
 		}
 	}
 
-	// bForceConnection: not used anymore
+	if (bForceConnection && !bServerAvailable)
+	{
+		// Execute a 'checkconnection' command to set bServerAvailable based on the connectivity of the server
+		TArray<FString> InfoMessages, ErrorMessages;
+		const bool bCommandSuccessful = PlasticSourceControlUtils::RunCommand(TEXT("checkconnection"), TArray<FString>(), TArray<FString>(), EConcurrency::Synchronous, InfoMessages, ErrorMessages);
+		bServerAvailable = bCommandSuccessful;
+		if (!bCommandSuccessful)
+		{
+			FMessageLog SourceControlLog("SourceControl");
+			for (const FString& ErrorMessage : ErrorMessages)
+			{
+				SourceControlLog.Error(FText::FromString(ErrorMessage));
+			}
+		}
+	}
 }
 
 void FPlasticSourceControlProvider::CheckPlasticAvailability()
@@ -56,6 +70,8 @@ void FPlasticSourceControlProvider::CheckPlasticAvailability()
 	FString PathToPlasticBinary = PlasticSourceControl.AccessSettings().GetBinaryPath();
 	if (PathToPlasticBinary.IsEmpty())
 	{
+		bPlasticAvailable = false;
+
 		// Try to find Plastic binary, and update settings accordingly
 		PathToPlasticBinary = PlasticSourceControlUtils::FindPlasticBinaryPath();
 		if (!PathToPlasticBinary.IsEmpty())
@@ -88,10 +104,6 @@ void FPlasticSourceControlProvider::CheckPlasticAvailability()
 			}
 		}
 	}
-	else
-	{
-		bPlasticAvailable = false;
-	}
 }
 
 void FPlasticSourceControlProvider::Close()
@@ -150,13 +162,13 @@ FText FPlasticSourceControlProvider::GetStatusText() const
 	return FText::Format( NSLOCTEXT("Status", "Provider: Plastic\nEnabledLabel", "Plastic SCM {PlasticScmVersion} (plugin v{PluginVersion})\nWorkspace: {WorkspaceName} ({WorkspacePath})\n{BranchName}\nChangeset: {ChangesetNumber}\nUser: {UserName}"), Args );
 }
 
-/** Quick check if source control is enabled */
+/** Quick check if source control is enabled. Specifically, it returns true if a source control provider is set (regardless of whether the provider is available) and false if no provider is set. So all providers except the stub DefaultSourceProvider will return true. */
 bool FPlasticSourceControlProvider::IsEnabled() const
 {
-	return bWorkspaceFound;
+	return true;
 }
 
-/** Quick check if source control is available for use (useful for server-based providers) */
+/** Quick check if source control is available for use (return whether the server is available or not) */
 bool FPlasticSourceControlProvider::IsAvailable() const
 {
 	return bServerAvailable;
@@ -309,14 +321,14 @@ void FPlasticSourceControlProvider::OutputCommandMessages(const FPlasticSourceCo
 {
 	FMessageLog SourceControlLog("SourceControl");
 
-	for (int32 ErrorIndex = 0; ErrorIndex < InCommand.ErrorMessages.Num(); ++ErrorIndex)
+	for (const FString& ErrorMessage : InCommand.ErrorMessages)
 	{
-		SourceControlLog.Error(FText::FromString(InCommand.ErrorMessages[ErrorIndex]));
+		SourceControlLog.Error(FText::FromString(ErrorMessage));
 	}
 
-	for (int32 InfoIndex = 0; InfoIndex < InCommand.InfoMessages.Num(); ++InfoIndex)
+	for (const FString& InfoMessage : InCommand.InfoMessages)
 	{
-		SourceControlLog.Info(FText::FromString(InCommand.InfoMessages[InfoIndex]));
+		SourceControlLog.Info(FText::FromString(InfoMessage));
 	}
 }
 
@@ -376,7 +388,14 @@ void FPlasticSourceControlProvider::Tick()
 			// dump any messages to output log
 			OutputCommandMessages(Command);
 
-			UE_LOG(LogSourceControl, Verbose, TEXT("%s processed in %.3lfs"), *Command.Operation->GetName().ToString(), (FPlatformTime::Seconds() - Command.StartTimestamp));
+			if (Command.Files.Num())
+			{
+				UE_LOG(LogSourceControl, Log, TEXT("%s of %d files processed in %.3lfs"), *Command.Operation->GetName().ToString(), Command.Files.Num(), (FPlatformTime::Seconds() - Command.StartTimestamp));
+			}
+			else
+			{
+				UE_LOG(LogSourceControl, Log, TEXT("%s processed in %.3lfs"), *Command.Operation->GetName().ToString(), (FPlatformTime::Seconds() - Command.StartTimestamp));
+			}
 
 			// run the completion delegate callback if we have one bound
 			ECommandResult::Type Result = Command.bCommandSuccessful ? ECommandResult::Succeeded : ECommandResult::Failed;
diff --git a/Engine/Plugins/Developer/PlasticSourceControl/Source/PlasticSourceControl/Private/PlasticSourceControlUtils.cpp b/Engine/Plugins/Developer/PlasticSourceControl/Source/PlasticSourceControl/Private/PlasticSourceControlUtils.cpp
--- a/Engine/Plugins/Developer/PlasticSourceControl/Source/PlasticSourceControl/Private/PlasticSourceControlUtils.cpp
+++ b/Engine/Plugins/Developer/PlasticSourceControl/Source/PlasticSourceControl/Private/PlasticSourceControlUtils.cpp
@@ -1536,7 +1536,11 @@ bool RunGetHistory(const bool bInUpdateHistory, TArray<FPlasticSourceControlStat
 	FString Results;
 	FString Errors;
 	TArray<FString> Parameters;
-	Parameters.Add(TEXT("--moveddeleted"));
+	// Detecting move and deletion is costly as it is implemented as two extra queries to the server; do it only when getting the history of the current branch
+	if (bInUpdateHistory)
+	{
+		Parameters.Add(TEXT("--moveddeleted"));
+	}
 	Parameters.Add(TEXT("--xml"));
 	Parameters.Add(TEXT("--encoding=\"utf-8\""));
 
@@ -1556,7 +1560,6 @@ bool RunGetHistory(const bool bInUpdateHistory, TArray<FPlasticSourceControlStat
 	if (Files.Num() > 0)
 	{
 		bResult = RunCommandInternal(TEXT("history"), Parameters, Files, EConcurrency::Synchronous, Results, Errors);
-		OutErrorMessages.Add(MoveTemp(Errors));
 		if (bResult)
 		{
 			FXmlFile XmlFile;
@@ -1566,6 +1569,10 @@ bool RunGetHistory(const bool bInUpdateHistory, TArray<FPlasticSourceControlStat
 				bResult = ParseHistoryResults(bInUpdateHistory, XmlFile, InOutStates);
 			}
 		}
+		if (!Errors.IsEmpty())
+		{
+			OutErrorMessages.Add(MoveTemp(Errors));
+		}
 	}
 
 	return bResult;
