Added changelists support to PlasticSCM plugin (v1.6.2)

PR #9414: [Source Control] Add Changelists support to PlasticSCM plugin (v1.6.2) (Contributed by [REDACTED_PERSON])

diff --git a/Engine/Plugins/Developer/PlasticSourceControl/PlasticSourceControl.uplugin b/Engine/Plugins/Developer/PlasticSourceControl/PlasticSourceControl.uplugin
--- a/Engine/Plugins/Developer/PlasticSourceControl/PlasticSourceControl.uplugin
+++ b/Engine/Plugins/Developer/PlasticSourceControl/PlasticSourceControl.uplugin
@@ -1 +1,27 @@
-{	"FileVersion": 3,	"Version": 51,	"VersionName": "1.5.1",	"FriendlyName": "Plastic SCM",	"Description": "Plastic source control management",	"Category": "Source Control",	"CreatedBy": "SRombauts for Codice Software.",	"CreatedByURL": "http://srombauts.github.com",	"DocsURL": "[REDACTED_PROFILE_URL]/UE4PlasticPlugin",	"MarketplaceURL": "",	"SupportURL": "https://www.plasticscm.com/support",	"EnabledByDefault": true,	"CanContainContent": false,	"IsBetaVersion": false,	"Installed": false,	"Modules": [		{			"Name": "PlasticSourceControl",			"Type": "Editor",			"LoadingPhase": "Default",			"PlatformAllowList": [				"Win64"			]		}	]}
\ No newline at end of file
+{
+	"FileVersion": 3,
+	"Version": 62,
+	"VersionName": "1.6.2",
+	"FriendlyName": "Plastic SCM",
+	"Description": "Plastic source control management",
+	"Category": "Source Control",
+	"CreatedBy": "SRombauts for Codice Software.",
+	"CreatedByURL": "http://srombauts.github.com",
+	"DocsURL": "[REDACTED_PROFILE_URL]/UEPlasticPlugin",
+	"MarketplaceURL": "",
+	"SupportURL": "https://www.plasticscm.com/support",
+	"EnabledByDefault": true,
+	"CanContainContent": false,
+	"IsBetaVersion": false,
+	"Installed": false,
+	"Modules": [
+		{
+			"Name": "PlasticSourceControl",
+			"Type": "Editor",
+			"LoadingPhase": "Default",
+			"PlatformAllowList": [
+				"Win64"
+			]
+		}
+	]
+}
\ No newline at end of file
diff --git a/Engine/Plugins/Developer/PlasticSourceControl/Source/PlasticSourceControl/PlasticSourceControl.Build.cs b/Engine/Plugins/Developer/PlasticSourceControl/Source/PlasticSourceControl/PlasticSourceControl.Build.cs
--- a/Engine/Plugins/Developer/PlasticSourceControl/Source/PlasticSourceControl/PlasticSourceControl.Build.cs
+++ b/Engine/Plugins/Developer/PlasticSourceControl/Source/PlasticSourceControl/PlasticSourceControl.Build.cs
@@ -20,6 +20,7 @@ public class PlasticSourceControl : ModuleRules
 				"XmlParser",
 				"Projects",
 				"AssetRegistry",
+				"DeveloperSettings",
 				"ToolMenus",
 			}
 		);
diff --git a/Engine/Plugins/Developer/PlasticSourceControl/Source/PlasticSourceControl/Private/IPlasticSourceControlWorker.h b/Engine/Plugins/Developer/PlasticSourceControl/Source/PlasticSourceControl/Private/IPlasticSourceControlWorker.h
--- a/Engine/Plugins/Developer/PlasticSourceControl/Source/PlasticSourceControl/Private/IPlasticSourceControlWorker.h
+++ b/Engine/Plugins/Developer/PlasticSourceControl/Source/PlasticSourceControl/Private/IPlasticSourceControlWorker.h
@@ -4,9 +4,32 @@
 
 #include "CoreMinimal.h"
 
+class FPlasticSourceControlProvider;
+
 class IPlasticSourceControlWorker
 {
 public:
+	// Implemented in PlasticSourceControlOperations.cpp with all the code for the workers
+	static void RegisterWorkers(FPlasticSourceControlProvider& PlasticSourceControlProvider);
+
+public:
+	explicit IPlasticSourceControlWorker(FPlasticSourceControlProvider& InSourceControlProvider)
+		: PlasticSourceControlProvider(InSourceControlProvider)
+	{
+	}
+
+	virtual ~IPlasticSourceControlWorker() = default;
+
+	FPlasticSourceControlProvider& GetProvider()
+	{
+		return PlasticSourceControlProvider;
+	}
+
+	const FPlasticSourceControlProvider& GetProvider() const
+	{
+		return PlasticSourceControlProvider;
+	}
+
 	/**
 	 * Name describing the work that this worker does. Used for factory method hookup.
 	 */
@@ -15,13 +38,16 @@ public:
 	/**
 	 * Function that actually does the work. Can be executed on another thread.
 	 */
-	virtual bool Execute( class FPlasticSourceControlCommand& InCommand ) = 0;
+	virtual bool Execute(class FPlasticSourceControlCommand& InCommand) = 0;
 
 	/**
 	 * Updates the state of any items after completion (if necessary). This is always executed on the main thread.
 	 * @returns true if states were updated
 	 */
 	virtual bool UpdateStates() = 0;
+
+private:
+	FPlasticSourceControlProvider& PlasticSourceControlProvider;
 };
 
 typedef TSharedRef<IPlasticSourceControlWorker, ESPMode::ThreadSafe> FPlasticSourceControlWorkerRef;
\ No newline at end of file
diff --git a/Engine/Plugins/Developer/PlasticSourceControl/Source/PlasticSourceControl/Private/PlasticSourceControlChangelist.cpp b/Engine/Plugins/Developer/PlasticSourceControl/Source/PlasticSourceControl/Private/PlasticSourceControlChangelist.cpp
new file mode 100644
--- /dev/null
+++ b/Engine/Plugins/Developer/PlasticSourceControl/Source/PlasticSourceControl/Private/PlasticSourceControlChangelist.cpp
@@ -0,0 +1,5 @@
+// Copyright Epic Games, Inc. All Rights Reserved.
+
+#include "PlasticSourceControlChangelist.h"
+
+const FPlasticSourceControlChangelist FPlasticSourceControlChangelist::DefaultChangelist(TEXT("Default"), true);
diff --git a/Engine/Plugins/Developer/PlasticSourceControl/Source/PlasticSourceControl/Private/PlasticSourceControlChangelist.h b/Engine/Plugins/Developer/PlasticSourceControl/Source/PlasticSourceControl/Private/PlasticSourceControlChangelist.h
new file mode 100644
--- /dev/null
+++ b/Engine/Plugins/Developer/PlasticSourceControl/Source/PlasticSourceControl/Private/PlasticSourceControlChangelist.h
@@ -0,0 +1,74 @@
+// Copyright Epic Games, Inc. All Rights Reserved.
+
+#pragma once
+
+#include "CoreMinimal.h"
+
+#include "ISourceControlChangelist.h"
+
+/**
+ * An abstraction of a changelist under source control: an identifier (a "name" in Plastic SCM)
+ */
+class FPlasticSourceControlChangelist : public ISourceControlChangelist
+{
+public:
+	FPlasticSourceControlChangelist() = default;
+
+	explicit FPlasticSourceControlChangelist(FString&& InChangelistName, const bool bInInitialized = false)
+		: ChangelistName(MoveTemp(InChangelistName))
+		, bInitialized(bInInitialized)
+	{
+	}
+
+	FPlasticSourceControlChangelist(const FPlasticSourceControlChangelist& InOther) = default;
+
+	bool operator==(const FPlasticSourceControlChangelist& InOther) const
+	{
+		return ChangelistName == InOther.ChangelistName;
+	}
+
+	bool operator!=(const FPlasticSourceControlChangelist& InOther) const
+	{
+		return ChangelistName != InOther.ChangelistName;
+	}
+
+	bool IsDefault() const
+	{
+		return ChangelistName == DefaultChangelist.ChangelistName;
+	}
+
+	void SetInitialized()
+	{
+		bInitialized = true;
+	}
+
+	bool IsInitialized() const
+	{
+		return bInitialized;
+	}
+
+	void Reset()
+	{
+		ChangelistName.Reset();
+		bInitialized = false;
+	}
+
+	friend FORCEINLINE uint32 GetTypeHash(const FPlasticSourceControlChangelist& InPlasticChangelist)
+	{
+		return GetTypeHash(InPlasticChangelist.ChangelistName);
+	}
+
+	FString GetName() const
+	{
+		return ChangelistName;
+	}
+
+public:
+	static const FPlasticSourceControlChangelist DefaultChangelist;
+
+private:
+	FString ChangelistName;
+	bool bInitialized = false;
+};
+
+typedef TSharedRef<class FPlasticSourceControlChangelist, ESPMode::ThreadSafe> FPlasticSourceControlChangelistRef;
diff --git a/Engine/Plugins/Developer/PlasticSourceControl/Source/PlasticSourceControl/Private/PlasticSourceControlChangelistState.cpp b/Engine/Plugins/Developer/PlasticSourceControl/Source/PlasticSourceControl/Private/PlasticSourceControlChangelistState.cpp
new file mode 100644
--- /dev/null
+++ b/Engine/Plugins/Developer/PlasticSourceControl/Source/PlasticSourceControl/Private/PlasticSourceControlChangelistState.cpp
@@ -0,0 +1,54 @@
+// Copyright Epic Games, Inc. All Rights Reserved.
+
+#include "PlasticSourceControlChangelistState.h"
+
+#define LOCTEXT_NAMESPACE "PlasticSourceControl.ChangelistState"
+
+FName FPlasticSourceControlChangelistState::GetIconName() const
+{
+	// Mimic P4V colors, returning the red icon if there are active file(s), the blue if the changelist is empty or all the files are shelved.
+	return Files.Num() > 0 ? FName("SourceControl.Changelist") : FName("SourceControl.ShelvedCHangelist");
+}
+
+FName FPlasticSourceControlChangelistState::GetSmallIconName() const
+{
+	return GetIconName();
+}
+
+FText FPlasticSourceControlChangelistState::GetDisplayText() const
+{
+	return FText::FromString(Changelist.GetName());
+}
+
+FText FPlasticSourceControlChangelistState::GetDescriptionText() const
+{
+	return FText::FromString(Description);
+}
+
+FText FPlasticSourceControlChangelistState::GetDisplayTooltip() const
+{
+	return LOCTEXT("Tooltip", "Tooltip");
+}
+
+const FDateTime& FPlasticSourceControlChangelistState::GetTimeStamp() const
+{
+	return TimeStamp;
+}
+
+const TArray<FSourceControlStateRef>& FPlasticSourceControlChangelistState::GetFilesStates() const
+{
+	return Files;
+}
+
+const TArray<FSourceControlStateRef>& FPlasticSourceControlChangelistState::GetShelvedFilesStates() const
+{
+	return ShelvedFiles;
+}
+
+FSourceControlChangelistRef FPlasticSourceControlChangelistState::GetChangelist() const
+{
+	FPlasticSourceControlChangelistRef ChangelistCopy = MakeShareable( new FPlasticSourceControlChangelist(Changelist));
+	return StaticCastSharedRef<ISourceControlChangelist>(ChangelistCopy);
+}
+
+#undef LOCTEXT_NAMESPACE
diff --git a/Engine/Plugins/Developer/PlasticSourceControl/Source/PlasticSourceControl/Private/PlasticSourceControlChangelistState.h b/Engine/Plugins/Developer/PlasticSourceControl/Source/PlasticSourceControl/Private/PlasticSourceControlChangelistState.h
new file mode 100644
--- /dev/null
+++ b/Engine/Plugins/Developer/PlasticSourceControl/Source/PlasticSourceControl/Private/PlasticSourceControlChangelistState.h
@@ -0,0 +1,81 @@
+// Copyright Epic Games, Inc. All Rights Reserved.
+
+#pragma once
+
+#include "CoreMinimal.h"
+
+#include "ISourceControlChangelistState.h"
+#include "PlasticSourceControlChangelist.h"
+
+/**
+ * An abstraction of the state of a pending changelist under source control: description and list of files
+ */
+class FPlasticSourceControlChangelistState : public ISourceControlChangelistState
+{
+public:
+	explicit FPlasticSourceControlChangelistState(const FPlasticSourceControlChangelist& InChangelist, const FString& InDescription = FString())
+		: Changelist(InChangelist)
+		, Description(InDescription)
+	{
+	}
+
+	explicit FPlasticSourceControlChangelistState(FPlasticSourceControlChangelist&& InChangelist, FString&& InDescription)
+		: Changelist(MoveTemp(InChangelist))
+		, Description(MoveTemp(InDescription))
+	{
+	}
+
+	/**
+	 * Get the name of the icon graphic we should use to display the state in a UI.
+	 * @returns the name of the icon to display
+	 */
+	virtual FName GetIconName() const override;
+
+	/**
+	 * Get the name of the small icon graphic we should use to display the state in a UI.
+	 * @returns the name of the icon to display
+	 */
+	virtual FName GetSmallIconName() const override;
+
+	/**
+	 * Get a text representation of the state
+	 * @returns	the text to display for this state
+	 */
+	virtual FText GetDisplayText() const override;
+
+	/**
+	 * Get a text representation of the state
+	 * @returns	the text to display for this state
+	 */
+	virtual FText GetDescriptionText() const override;
+
+	/**
+	 * Get a tooltip to describe this state
+	 * @returns	the text to display for this states tooltip
+	 */
+	virtual FText GetDisplayTooltip() const override;
+
+	/**
+	 * Get the timestamp of the last update that was made to this state.
+	 * @returns	the timestamp of the last update
+	 */
+	virtual const FDateTime& GetTimeStamp() const override;
+
+	virtual const TArray<FSourceControlStateRef>& GetFilesStates() const override;
+
+	virtual const TArray<FSourceControlStateRef>& GetShelvedFilesStates() const override;
+
+	virtual FSourceControlChangelistRef GetChangelist() const override;
+
+public:
+	FPlasticSourceControlChangelist Changelist;
+
+	FString Description;
+
+	TArray<FSourceControlStateRef> Files;
+
+	TArray<FSourceControlStateRef> ShelvedFiles; // TODO: shelves are not yet implemented
+
+	/** The timestamp of the last update */
+	FDateTime TimeStamp;
+};
diff --git a/Engine/Plugins/Developer/PlasticSourceControl/Source/PlasticSourceControl/Private/PlasticSourceControlCommand.cpp b/Engine/Plugins/Developer/PlasticSourceControl/Source/PlasticSourceControl/Private/PlasticSourceControlCommand.cpp
--- a/Engine/Plugins/Developer/PlasticSourceControl/Source/PlasticSourceControl/Private/PlasticSourceControlCommand.cpp
+++ b/Engine/Plugins/Developer/PlasticSourceControl/Source/PlasticSourceControl/Private/PlasticSourceControlCommand.cpp
@@ -20,9 +20,9 @@ FPlasticSourceControlCommand::FPlasticSourceControlCommand(const FSourceControlO
 {
 	// grab the providers settings here, so we don't access them once the worker thread is launched
 	check(IsInGameThread());
-	const FPlasticSourceControlModule& PlasticSourceControl = FModuleManager::LoadModuleChecked<FPlasticSourceControlModule>( "PlasticSourceControl" );
-	PathToWorkspaceRoot = PlasticSourceControl.GetProvider().GetPathToWorkspaceRoot();
-	ChangesetNumber = PlasticSourceControl.GetProvider().GetChangesetNumber();
+	const FPlasticSourceControlProvider& Provider = FPlasticSourceControlModule::Get().GetProvider();
+	PathToWorkspaceRoot = Provider.GetPathToWorkspaceRoot();
+	ChangesetNumber = Provider.GetChangesetNumber();
 }
 
 bool FPlasticSourceControlCommand::DoWork()
diff --git a/Engine/Plugins/Developer/PlasticSourceControl/Source/PlasticSourceControl/Private/PlasticSourceControlCommand.h b/Engine/Plugins/Developer/PlasticSourceControl/Source/PlasticSourceControl/Private/PlasticSourceControlCommand.h
--- a/Engine/Plugins/Developer/PlasticSourceControl/Source/PlasticSourceControl/Private/PlasticSourceControlCommand.h
+++ b/Engine/Plugins/Developer/PlasticSourceControl/Source/PlasticSourceControl/Private/PlasticSourceControlCommand.h
@@ -7,13 +7,14 @@
 #include "ISourceControlProvider.h"
 #include "IPlasticSourceControlWorker.h"
 
+#include "PlasticSourceControlChangelist.h"
+
 /**
  * Used to execute Plastic commands multi-threaded.
  */
 class FPlasticSourceControlCommand : public IQueuedWork
 {
 public:
-
 	FPlasticSourceControlCommand(const FSourceControlOperationRef& InOperation, const FPlasticSourceControlWorkerRef& InWorker, const FSourceControlOperationComplete& InOperationCompleteDelegate = FSourceControlOperationComplete() );
 
 	/**
@@ -83,11 +84,14 @@ public:
 	const double StartTimestamp;
 
 	/** Files to perform this operation on */
-	TArray< FString > Files;
+	TArray<FString> Files;
+
+	/** Changelist to perform this operation on */
+	FPlasticSourceControlChangelist Changelist;
 
-	/**Info and/or warning message message storage*/
-	TArray< FString > InfoMessages;
+	/**Info and/or warning message storage*/
+	TArray<FString> InfoMessages;
 
 	/**Potential error message storage*/
-	TArray< FString > ErrorMessages;
+	TArray<FString> ErrorMessages;
 };
\ No newline at end of file
diff --git a/Engine/Plugins/Developer/PlasticSourceControl/Source/PlasticSourceControl/Private/PlasticSourceControlConsole.cpp b/Engine/Plugins/Developer/PlasticSourceControl/Source/PlasticSourceControl/Private/PlasticSourceControlConsole.cpp
--- a/Engine/Plugins/Developer/PlasticSourceControl/Source/PlasticSourceControl/Private/PlasticSourceControlConsole.cpp
+++ b/Engine/Plugins/Developer/PlasticSourceControl/Source/PlasticSourceControl/Private/PlasticSourceControlConsole.cpp
@@ -2,20 +2,21 @@
 
 #include "PlasticSourceControlConsole.h"
 
-#include "HAL/IConsoleManager.h"
 #include "ISourceControlModule.h"
 
 #include "PlasticSourceControlUtils.h"
 
 void FPlasticSourceControlConsole::Register()
 {
-	CmConsoleCommand = MakeUnique<FAutoConsoleCommand>(
-		TEXT("cm"),
-		TEXT("PlasticSCM Command Line Interface.\n")
-		TEXT("Run any 'cm' command directly from the Unreal Editor Console.\n")
-		TEXT("Type 'cm showcommands' to get a command list."),
-		FConsoleCommandWithArgsDelegate::CreateRaw(this, &FPlasticSourceControlConsole::ExecutePlasticConsoleCommand)
-	);
+	if (!CmConsoleCommand.IsValid())
+	{
+		CmConsoleCommand = MakeUnique<FAutoConsoleCommand>(
+			TEXT("cm"),
+			TEXT("PlasticSCM Command Line Interface.\n")
+			TEXT("Run any 'cm' command directly from the Unreal Editor Console.\n")
+			TEXT("Type 'cm showcommands' to get a command list."),
+			FConsoleCommandWithArgsDelegate::CreateRaw(this, &FPlasticSourceControlConsole::ExecutePlasticConsoleCommand));
+	}
 }
 
 void FPlasticSourceControlConsole::Unregister()
@@ -25,7 +26,7 @@ void FPlasticSourceControlConsole::Unregister()
 
 void FPlasticSourceControlConsole::ExecutePlasticConsoleCommand(const TArray<FString>& a_args)
 {
-	// If called with no argument, explicitely call "cm help" instead to mimic the cm CLI behavior.
+	// If called with no argument, explicitly call "cm help" instead to mimic the cm CLI behavior.
 	if (a_args.Num() < 1)
 	{
 		ExecutePlasticConsoleCommand(TArray<FString>({TEXT("help")}));
diff --git a/Engine/Plugins/Developer/PlasticSourceControl/Source/PlasticSourceControl/Private/PlasticSourceControlConsole.h b/Engine/Plugins/Developer/PlasticSourceControl/Source/PlasticSourceControl/Private/PlasticSourceControlConsole.h
--- a/Engine/Plugins/Developer/PlasticSourceControl/Source/PlasticSourceControl/Private/PlasticSourceControlConsole.h
+++ b/Engine/Plugins/Developer/PlasticSourceControl/Source/PlasticSourceControl/Private/PlasticSourceControlConsole.h
@@ -5,8 +5,6 @@
 #include "CoreMinimal.h"
 #include "HAL/IConsoleManager.h"
 
-class FAutoConsoleCommand;
-
 /**
  * Editor only console commands.
  *
diff --git a/Engine/Plugins/Developer/PlasticSourceControl/Source/PlasticSourceControl/Private/PlasticSourceControlMenu.cpp b/Engine/Plugins/Developer/PlasticSourceControl/Source/PlasticSourceControl/Private/PlasticSourceControlMenu.cpp
--- a/Engine/Plugins/Developer/PlasticSourceControl/Source/PlasticSourceControl/Private/PlasticSourceControlMenu.cpp
+++ b/Engine/Plugins/Developer/PlasticSourceControl/Source/PlasticSourceControl/Private/PlasticSourceControlMenu.cpp
@@ -20,6 +20,7 @@
 
 #include "PackageTools.h"
 #include "FileHelpers.h"
+#include "ISettingsModule.h"
 
 #include "Logging/MessageLog.h"
 
@@ -70,7 +71,7 @@ bool FPlasticSourceControlMenu::SaveDirtyPackages()
 
 	bool bSaved = FEditorFileUtils::SaveDirtyPackages(bPromptUserToSave, bSaveMapPackages, bSaveContentPackages, bFastSave, bNotifyNoPackagesSaved, bCanBeDeclined, &bHadPackagesToSave);
 
-	// bSaved can be true if the user selects to not save an asset by unchecking it and clicking "save"
+	// bSaved can be true if the user selects to not save an asset by un-checking it and clicking "save"
 	if (bSaved)
 	{
 		TArray<UPackage*> DirtyPackages;
@@ -175,8 +176,7 @@ void FPlasticSourceControlMenu::SyncProjectClicked()
 			PackagesToReload = UnlinkPackages(ListAllPackages());
 
 			// Launch a "Sync" operation
-			FPlasticSourceControlModule& PlasticSourceControl = FModuleManager::LoadModuleChecked<FPlasticSourceControlModule>("PlasticSourceControl");
-			FPlasticSourceControlProvider& Provider = PlasticSourceControl.GetProvider();
+			FPlasticSourceControlProvider& Provider = FPlasticSourceControlModule::Get().GetProvider();
 			TSharedRef<FSync, ESPMode::ThreadSafe> SyncOperation = ISourceControlOperation::Create<FSync>();
 			const ECommandResult::Type Result = Provider.Execute(SyncOperation, TArray<FString>(), EConcurrency::Asynchronous, FSourceControlOperationComplete::CreateRaw(this, &FPlasticSourceControlMenu::OnSourceControlOperationComplete));
 			if (Result == ECommandResult::Succeeded)
@@ -211,8 +211,7 @@ void FPlasticSourceControlMenu::RevertUnchangedClicked()
 	if (!OperationInProgressNotification.IsValid())
 	{
 		// Launch a "RevertUnchanged" Operation
-		FPlasticSourceControlModule& PlasticSourceControl = FModuleManager::LoadModuleChecked<FPlasticSourceControlModule>("PlasticSourceControl");
-		FPlasticSourceControlProvider& Provider = PlasticSourceControl.GetProvider();
+		FPlasticSourceControlProvider& Provider = FPlasticSourceControlModule::Get().GetProvider();
 		TSharedRef<FPlasticRevertUnchanged, ESPMode::ThreadSafe> RevertUnchangedOperation = ISourceControlOperation::Create<FPlasticRevertUnchanged>();
 		TArray<FString> WorkspaceRoot;
 		WorkspaceRoot.Add(Provider.GetPathToWorkspaceRoot()); // Revert the root of the workspace (needs an absolute path)
@@ -252,8 +251,7 @@ void FPlasticSourceControlMenu::RevertAllClicked()
 				PackagesToReload = UnlinkPackages(ListAllPackages());
 
 				// Launch a "RevertAll" Operation
-				FPlasticSourceControlModule& PlasticSourceControl = FModuleManager::LoadModuleChecked<FPlasticSourceControlModule>("PlasticSourceControl");
-				FPlasticSourceControlProvider& Provider = PlasticSourceControl.GetProvider();
+				FPlasticSourceControlProvider& Provider = FPlasticSourceControlModule::Get().GetProvider();
 				TSharedRef<FPlasticRevertAll, ESPMode::ThreadSafe> RevertAllOperation = ISourceControlOperation::Create<FPlasticRevertAll>();
 				TArray<FString> WorkspaceRoot;
 				WorkspaceRoot.Add(Provider.GetPathToWorkspaceRoot()); // Revert the root of the workspace (needs an absolute path)
@@ -291,8 +289,7 @@ void FPlasticSourceControlMenu::RefreshClicked()
 	if (!OperationInProgressNotification.IsValid())
 	{
 		// Launch an "UpdateStatus" Operation
-		FPlasticSourceControlModule& PlasticSourceControl = FModuleManager::LoadModuleChecked<FPlasticSourceControlModule>("PlasticSourceControl");
-		FPlasticSourceControlProvider& Provider = PlasticSourceControl.GetProvider();
+		FPlasticSourceControlProvider& Provider = FPlasticSourceControlModule::Get().GetProvider();
 		TSharedRef<FUpdateStatus, ESPMode::ThreadSafe> RefreshOperation = ISourceControlOperation::Create<FUpdateStatus>();
 		// This is the flag used by the Content Browser's "Checkout" filter to trigger a full status update
 		RefreshOperation->SetGetOpenedOnly(true);
@@ -316,7 +313,31 @@ void FPlasticSourceControlMenu::RefreshClicked()
 	}
 }
 
-void FPlasticSourceControlMenu::VisitDocsURLClicked()
+void FPlasticSourceControlMenu::ShowSourceControlEditorPreferences() const
+{
+	if (ISettingsModule* SettingsModule = FModuleManager::GetModulePtr<ISettingsModule>("Settings"))
+	{
+		SettingsModule->ShowViewer("Editor", "General", "LoadingSaving");
+	}
+}
+
+void FPlasticSourceControlMenu::ShowSourceControlProjectSettings() const
+{
+	if (ISettingsModule* SettingsModule = FModuleManager::GetModulePtr<ISettingsModule>("Settings"))
+	{
+		SettingsModule->ShowViewer("Project", "Editor", "SourceControlPreferences");
+	}
+}
+
+void FPlasticSourceControlMenu::ShowSourceControlPlasticScmProjectSettings() const
+{
+	if (ISettingsModule* SettingsModule = FModuleManager::GetModulePtr<ISettingsModule>("Settings"))
+	{
+		SettingsModule->ShowViewer("Project", "Editor", "PlasticSourceControlProjectSettings");
+	}
+}
+
+void FPlasticSourceControlMenu::VisitDocsURLClicked() const
 {
 	// Grab the URL from the uplugin file
 	const TSharedPtr<IPlugin> Plugin = IPluginManager::Get().FindPlugin(TEXT("PlasticSourceControl"));
@@ -326,7 +347,7 @@ void FPlasticSourceControlMenu::VisitDocsURLClicked()
 	}
 }
 
-void FPlasticSourceControlMenu::VisitSupportURLClicked()
+void FPlasticSourceControlMenu::VisitSupportURLClicked() const
 {
 	// Grab the URL from the uplugin file
 	const TSharedPtr<IPlugin> Plugin = IPluginManager::Get().FindPlugin(TEXT("PlasticSourceControl"));
@@ -413,17 +434,6 @@ void FPlasticSourceControlMenu::OnSourceControlOperationComplete(const FSourceCo
 
 void FPlasticSourceControlMenu::AddMenuExtension(FToolMenuSection& Menu)
 {
-	Menu.AddMenuEntry(
-		"PlasticSync",
-		LOCTEXT("PlasticSync",			"Sync/Update Workspace"),
-		LOCTEXT("PlasticSyncTooltip",	"Update all files in the workspace to the latest version."),
-		FSlateIcon(FAppStyle::GetAppStyleSetName(), "SourceControl.Actions.Sync"),
-		FUIAction(
-			FExecuteAction::CreateRaw(this, &FPlasticSourceControlMenu::SyncProjectClicked),
-			FCanExecuteAction::CreateRaw(this, &FPlasticSourceControlMenu::IsSourceControlConnected)
-		)
-	);
-
 	Menu.AddMenuEntry(
 		"PlasticRevertUnchanged",
 		LOCTEXT("PlasticRevertUnchanged",			"Revert Unchanged"),
@@ -436,23 +446,44 @@ void FPlasticSourceControlMenu::AddMenuExtension(FToolMenuSection& Menu)
 	);
 
 	Menu.AddMenuEntry(
-		"PlasticRevertAll",
-		LOCTEXT("PlasticRevertAll",			"Revert All"),
-		LOCTEXT("PlasticRevertAllTooltip",	"Revert all files in the workspace to their controlled/unchanged state."),
-		FSlateIcon(FAppStyle::GetAppStyleSetName(), "SourceControl.Actions.Revert"),
+		"PlasticRefresh",
+		LOCTEXT("PlasticRefresh",			"Refresh"),
+		LOCTEXT("PlasticRefreshTooltip",	"Update the source control status of all files in the workspace."),
+		FSlateIcon(FAppStyle::GetAppStyleSetName(), "SourceControl.Actions.Refresh"),
 		FUIAction(
-			FExecuteAction::CreateRaw(this, &FPlasticSourceControlMenu::RevertAllClicked),
+			FExecuteAction::CreateRaw(this, &FPlasticSourceControlMenu::RefreshClicked),
 			FCanExecuteAction()
 		)
 	);
 
 	Menu.AddMenuEntry(
-		"PlasticRefresh",
-		LOCTEXT("PlasticRefresh",			"Refresh"),
-		LOCTEXT("PlasticRefreshTooltip",	"Update the source control status of all files in the workspace."),
-		FSlateIcon(FAppStyle::GetAppStyleSetName(), "SourceControl.Actions.Refresh"),
+		"SourceControlEditorPreferences",
+		LOCTEXT("SourceControlEditorPreferences", "Editor Preferences - Source Control"),
+		LOCTEXT("SourceControlEditorPreferencesTooltip", "Open the Load & Save section with Source Control in the Editor Preferences."),
+		FSlateIcon(FAppStyle::GetAppStyleSetName(), "EditorPreferences.TabIcon"),
 		FUIAction(
-			FExecuteAction::CreateRaw(this, &FPlasticSourceControlMenu::RefreshClicked),
+			FExecuteAction::CreateRaw(this, &FPlasticSourceControlMenu::ShowSourceControlEditorPreferences),
+			FCanExecuteAction()
+		)
+	);
+
+	Menu.AddMenuEntry(
+		"SourceControlProjectSettings",
+		LOCTEXT("SourceControlProjectSettings",			"Project Settings - Source Control"),
+		LOCTEXT("SourceControlProjectSettingsTooltip",	"Open the Source Control section in the Project Settings."),
+		FSlateIcon(FAppStyle::GetAppStyleSetName(), "ProjectSettings.TabIcon"),
+		FUIAction(
+			FExecuteAction::CreateRaw(this, &FPlasticSourceControlMenu::ShowSourceControlProjectSettings),
+			FCanExecuteAction()
+		)
+	);
+	Menu.AddMenuEntry(
+		"PlasticProjectSettings",
+		LOCTEXT("PlasticProjectSettings",			"Project Settings - Source Control - Plastic SCM"),
+		LOCTEXT("PlasticProjectSettingsTooltip",	"Open the Plastic SCM section in the Project Settings."),
+		FSlateIcon(FAppStyle::GetAppStyleSetName(), "ProjectSettings.TabIcon"),
+		FUIAction(
+			FExecuteAction::CreateRaw(this, &FPlasticSourceControlMenu::ShowSourceControlPlasticScmProjectSettings),
 			FCanExecuteAction()
 		)
 	);
diff --git a/Engine/Plugins/Developer/PlasticSourceControl/Source/PlasticSourceControl/Private/PlasticSourceControlMenu.h b/Engine/Plugins/Developer/PlasticSourceControl/Source/PlasticSourceControl/Private/PlasticSourceControlMenu.h
--- a/Engine/Plugins/Developer/PlasticSourceControl/Source/PlasticSourceControl/Private/PlasticSourceControlMenu.h
+++ b/Engine/Plugins/Developer/PlasticSourceControl/Source/PlasticSourceControl/Private/PlasticSourceControlMenu.h
@@ -19,8 +19,11 @@ public:
 	void RevertUnchangedClicked();
 	void RevertAllClicked();
 	void RefreshClicked();
-	void VisitDocsURLClicked();
-	void VisitSupportURLClicked();
+	void ShowSourceControlEditorPreferences() const;
+	void ShowSourceControlProjectSettings() const;
+	void ShowSourceControlPlasticScmProjectSettings() const;
+	void VisitDocsURLClicked() const;
+	void VisitSupportURLClicked() const;
 
 private:
 	bool IsSourceControlConnected() const;
diff --git a/Engine/Plugins/Developer/PlasticSourceControl/Source/PlasticSourceControl/Private/PlasticSourceControlModule.cpp b/Engine/Plugins/Developer/PlasticSourceControl/Source/PlasticSourceControl/Private/PlasticSourceControlModule.cpp
--- a/Engine/Plugins/Developer/PlasticSourceControl/Source/PlasticSourceControl/Private/PlasticSourceControlModule.cpp
+++ b/Engine/Plugins/Developer/PlasticSourceControl/Source/PlasticSourceControl/Private/PlasticSourceControlModule.cpp
@@ -1,40 +1,20 @@
 // Copyright Epic Games, Inc. All Rights Reserved.
 
 #include "PlasticSourceControlModule.h"
-#include "PlasticSourceControlOperations.h"
-#include "Misc/App.h"
+
+#include "IPlasticSourceControlWorker.h"
 #include "Features/IModularFeatures.h"
+#include "Misc/App.h"
 
 #define LOCTEXT_NAMESPACE "PlasticSourceControl"
 
-template<typename Type>
-static FPlasticSourceControlWorkerRef CreateWorker()
-{
-	return MakeShareable( new Type() );
-}
-
 void FPlasticSourceControlModule::StartupModule()
 {
-	// Register our operations (implemented in PlasticSourceControlOperations.cpp by subclassing from Engine\Source\Developer\SourceControl\Public\SourceControlOperations.h)
-	PlasticSourceControlProvider.RegisterWorker("Connect", FGetPlasticSourceControlWorker::CreateStatic( &CreateWorker<FPlasticConnectWorker>));
-	PlasticSourceControlProvider.RegisterWorker("CheckOut", FGetPlasticSourceControlWorker::CreateStatic(&CreateWorker<FPlasticCheckOutWorker>));
-	PlasticSourceControlProvider.RegisterWorker("UpdateStatus", FGetPlasticSourceControlWorker::CreateStatic( &CreateWorker<FPlasticUpdateStatusWorker>));
-	PlasticSourceControlProvider.RegisterWorker("MarkForAdd", FGetPlasticSourceControlWorker::CreateStatic(&CreateWorker<FPlasticMarkForAddWorker>));
-	PlasticSourceControlProvider.RegisterWorker("Delete", FGetPlasticSourceControlWorker::CreateStatic(&CreateWorker<FPlasticDeleteWorker>));
-	PlasticSourceControlProvider.RegisterWorker("Revert", FGetPlasticSourceControlWorker::CreateStatic(&CreateWorker<FPlasticRevertWorker>));
-	PlasticSourceControlProvider.RegisterWorker("RevertUnchanged", FGetPlasticSourceControlWorker::CreateStatic(&CreateWorker<FPlasticRevertUnchangedWorker>));
-	PlasticSourceControlProvider.RegisterWorker("RevertAll", FGetPlasticSourceControlWorker::CreateStatic(&CreateWorker<FPlasticRevertAllWorker>));
-	PlasticSourceControlProvider.RegisterWorker("MakeWorkspace", FGetPlasticSourceControlWorker::CreateStatic(&CreateWorker<FPlasticMakeWorkspaceWorker>));
-	PlasticSourceControlProvider.RegisterWorker("Sync", FGetPlasticSourceControlWorker::CreateStatic( &CreateWorker<FPlasticSyncWorker>));
-	PlasticSourceControlProvider.RegisterWorker("CheckIn", FGetPlasticSourceControlWorker::CreateStatic(&CreateWorker<FPlasticCheckInWorker>));
-	PlasticSourceControlProvider.RegisterWorker("Copy", FGetPlasticSourceControlWorker::CreateStatic(&CreateWorker<FPlasticCopyWorker>));
-	PlasticSourceControlProvider.RegisterWorker("Resolve", FGetPlasticSourceControlWorker::CreateStatic(&CreateWorker<FPlasticResolveWorker>));
-
-	// load our settings
-	PlasticSourceControlSettings.LoadSettings();
+	// Register our operations (implemented in PlasticSourceControlOperations.cpp by sub-classing from Engine\Source\Developer\SourceControl\Public\SourceControlOperations.h)
+	IPlasticSourceControlWorker::RegisterWorkers(PlasticSourceControlProvider);
 
 	// Bind our source control provider to the editor
-	IModularFeatures::Get().RegisterModularFeature( "SourceControl", &PlasticSourceControlProvider );
+	IModularFeatures::Get().RegisterModularFeature("SourceControl", &PlasticSourceControlProvider);
 }
 
 void FPlasticSourceControlModule::ShutdownModule()
@@ -46,16 +26,6 @@ void FPlasticSourceControlModule::ShutdownModule()
 	IModularFeatures::Get().UnregisterModularFeature("SourceControl", &PlasticSourceControlProvider);
 }
 
-void FPlasticSourceControlModule::SaveSettings()
-{
-	if (FApp::IsUnattended() || IsRunningCommandlet())
-	{
-		return;
-	}
-
-	PlasticSourceControlSettings.SaveSettings();
-}
-
 IMPLEMENT_MODULE(FPlasticSourceControlModule, PlasticSourceControl);
 
 #undef LOCTEXT_NAMESPACE
diff --git a/Engine/Plugins/Developer/PlasticSourceControl/Source/PlasticSourceControl/Private/PlasticSourceControlModule.h b/Engine/Plugins/Developer/PlasticSourceControl/Source/PlasticSourceControl/Private/PlasticSourceControlModule.h
--- a/Engine/Plugins/Developer/PlasticSourceControl/Source/PlasticSourceControl/Private/PlasticSourceControlModule.h
+++ b/Engine/Plugins/Developer/PlasticSourceControl/Source/PlasticSourceControl/Private/PlasticSourceControlModule.h
@@ -5,13 +5,12 @@
 #include "CoreMinimal.h"
 #include "Modules/ModuleInterface.h"
 #include "Modules/ModuleManager.h"
-#include "PlasticSourceControlSettings.h"
 #include "PlasticSourceControlProvider.h"
 
 /**
- * PlasticSourceControl is a simple Plastic SCM Source Control Plugin for Unreal Engine
+ * PlasticSourceControl is the official Plastic SCM Source Control Plugin for Unreal Engine
  *
- * Written and contributed by [REDACTED_PERSON] ([REDACTED_EMAIL])
+ * Written and contributed by [REDACTED_PERSON] ([REDACTED_EMAIL]) for Codice Software
  */
 class FPlasticSourceControlModule : public IModuleInterface
 {
@@ -20,19 +19,6 @@ public:
 	virtual void StartupModule() override;
 	virtual void ShutdownModule() override;
 
-	/** Access the Plastic source control settings */
-	FPlasticSourceControlSettings& AccessSettings()
-	{
-		return PlasticSourceControlSettings;
-	}
-	const FPlasticSourceControlSettings& AccessSettings() const
-	{
-		return PlasticSourceControlSettings;
-	}
-
-	/** Save the Plastic source control settings */
-	void SaveSettings();
-
 	/** Access the Plastic source control provider */
 	FPlasticSourceControlProvider& GetProvider()
 	{
@@ -43,10 +29,18 @@ public:
 		return PlasticSourceControlProvider;
 	}
 
+	/**
+	 * Singleton-like access to this module's interface.  This is just for convenience!
+	 * Beware of calling this during the shutdown phase, though.  Your module might have been unloaded already.
+	 *
+	 * @return Returns singleton instance, loading the module on demand if needed
+	 */
+	static inline FPlasticSourceControlModule& Get()
+	{
+		return FModuleManager::GetModuleChecked<FPlasticSourceControlModule>("PlasticSourceControl");
+	}
+
 private:
 	/** The Plastic source control provider */
 	FPlasticSourceControlProvider PlasticSourceControlProvider;
-
-	/** The settings for Plastic source control */
-	FPlasticSourceControlSettings PlasticSourceControlSettings;
 };
\ No newline at end of file
diff --git a/Engine/Plugins/Developer/PlasticSourceControl/Source/PlasticSourceControl/Private/PlasticSourceControlOperations.cpp b/Engine/Plugins/Developer/PlasticSourceControl/Source/PlasticSourceControl/Private/PlasticSourceControlOperations.cpp
--- a/Engine/Plugins/Developer/PlasticSourceControl/Source/PlasticSourceControl/Private/PlasticSourceControlOperations.cpp
+++ b/Engine/Plugins/Developer/PlasticSourceControl/Source/PlasticSourceControl/Private/PlasticSourceControlOperations.cpp
@@ -1,10 +1,12 @@
 // Copyright Epic Games, Inc. All Rights Reserved.
 
 #include "PlasticSourceControlOperations.h"
-#include "PlasticSourceControlSettings.h"
-#include "PlasticSourceControlState.h"
+
 #include "PlasticSourceControlCommand.h"
 #include "PlasticSourceControlModule.h"
+#include "PlasticSourceControlProvider.h"
+#include "PlasticSourceControlSettings.h"
+#include "PlasticSourceControlState.h"
 #include "PlasticSourceControlUtils.h"
 
 #include "AssetRegistry/AssetRegistryModule.h"
@@ -13,9 +15,49 @@
 #include "Misc/Paths.h"
 #include "SourceControlOperations.h"
 #include "ISourceControlModule.h"
+#include "Algo/NoneOf.h"
+
+#include "Runtime/Launch/Resources/Version.h"
 
 #define LOCTEXT_NAMESPACE "PlasticSourceControl"
 
+// 11.0.16.7248 add support for --descriptionfile for multi-line descriptions and support for special characters
+// https://www.plasticscm.com/download/releasenotes/11.0.16.7248
+static const FSoftwareVersion s_NewChangelistFileArgsPlasticScmVersion(11, 0, 16, 7248);
+
+template<typename Type>
+static FPlasticSourceControlWorkerRef InstantiateWorker(FPlasticSourceControlProvider& PlasticSourceControlProvider)
+{
+	return MakeShareable(new Type(PlasticSourceControlProvider));
+}
+
+void IPlasticSourceControlWorker::RegisterWorkers(FPlasticSourceControlProvider& PlasticSourceControlProvider)
+{
+	// Register our operations (implemented in PlasticSourceControlOperations.cpp by sub-classing from Engine\Source\Developer\SourceControl\Public\SourceControlOperations.h)
+	PlasticSourceControlProvider.RegisterWorker("Connect", FGetPlasticSourceControlWorker::CreateStatic(&InstantiateWorker<FPlasticConnectWorker>));
+	PlasticSourceControlProvider.RegisterWorker("CheckOut", FGetPlasticSourceControlWorker::CreateStatic(&InstantiateWorker<FPlasticCheckOutWorker>));
+	PlasticSourceControlProvider.RegisterWorker("UpdateStatus", FGetPlasticSourceControlWorker::CreateStatic(&InstantiateWorker<FPlasticUpdateStatusWorker>));
+	PlasticSourceControlProvider.RegisterWorker("MarkForAdd", FGetPlasticSourceControlWorker::CreateStatic(&InstantiateWorker<FPlasticMarkForAddWorker>));
+	PlasticSourceControlProvider.RegisterWorker("Delete", FGetPlasticSourceControlWorker::CreateStatic(&InstantiateWorker<FPlasticDeleteWorker>));
+	PlasticSourceControlProvider.RegisterWorker("Revert", FGetPlasticSourceControlWorker::CreateStatic(&InstantiateWorker<FPlasticRevertWorker>));
+	PlasticSourceControlProvider.RegisterWorker("RevertUnchanged", FGetPlasticSourceControlWorker::CreateStatic(&InstantiateWorker<FPlasticRevertUnchangedWorker>));
+	PlasticSourceControlProvider.RegisterWorker("RevertAll", FGetPlasticSourceControlWorker::CreateStatic(&InstantiateWorker<FPlasticRevertAllWorker>));
+	PlasticSourceControlProvider.RegisterWorker("MakeWorkspace", FGetPlasticSourceControlWorker::CreateStatic(&InstantiateWorker<FPlasticMakeWorkspaceWorker>));
+	PlasticSourceControlProvider.RegisterWorker("Sync", FGetPlasticSourceControlWorker::CreateStatic(&InstantiateWorker<FPlasticSyncWorker>));
+	PlasticSourceControlProvider.RegisterWorker("CheckIn", FGetPlasticSourceControlWorker::CreateStatic(&InstantiateWorker<FPlasticCheckInWorker>));
+	PlasticSourceControlProvider.RegisterWorker("Copy", FGetPlasticSourceControlWorker::CreateStatic(&InstantiateWorker<FPlasticCopyWorker>));
+	PlasticSourceControlProvider.RegisterWorker("Resolve", FGetPlasticSourceControlWorker::CreateStatic(&InstantiateWorker<FPlasticResolveWorker>));
+
+#if ENGINE_MAJOR_VERSION == 5
+	PlasticSourceControlProvider.RegisterWorker("UpdateChangelistsStatus", FGetPlasticSourceControlWorker::CreateStatic(&InstantiateWorker<FPlasticGetPendingChangelistsWorker>));
+	PlasticSourceControlProvider.RegisterWorker("NewChangelist", FGetPlasticSourceControlWorker::CreateStatic(&InstantiateWorker<FPlasticNewChangelistWorker>));
+	PlasticSourceControlProvider.RegisterWorker("DeleteChangelist", FGetPlasticSourceControlWorker::CreateStatic(&InstantiateWorker<FPlasticDeleteChangelistWorker>));
+	PlasticSourceControlProvider.RegisterWorker("EditChangelist", FGetPlasticSourceControlWorker::CreateStatic(&InstantiateWorker<FPlasticEditChangelistWorker>));
+	PlasticSourceControlProvider.RegisterWorker("MoveToChangelist", FGetPlasticSourceControlWorker::CreateStatic(&InstantiateWorker<FPlasticReopenWorker>));
+#endif
+}
+
+
 FName FPlasticRevertUnchanged::GetName() const
 {
 	return "RevertUnchanged";
@@ -64,48 +106,60 @@ FName FPlasticConnectWorker::GetName() const
 
 bool FPlasticConnectWorker::Execute(FPlasticSourceControlCommand& InCommand)
 {
+	TRACE_CPUPROFILER_EVENT_SCOPE(FPlasticConnectWorker::Execute);
+
 	check(InCommand.Operation->GetName() == GetName());
 	TSharedRef<FConnect, ESPMode::ThreadSafe> Operation = StaticCastSharedRef<FConnect>(InCommand.Operation);
 
-	const FPlasticSourceControlModule& PlasticSourceControl = FModuleManager::GetModuleChecked<FPlasticSourceControlModule>("PlasticSourceControl");
-	if (PlasticSourceControl.GetProvider().IsPlasticAvailable())
+	if (GetProvider().IsPlasticAvailable())
 	{
 		// Get workspace name
-		InCommand.bCommandSuccessful = PlasticSourceControlUtils::GetWorkspaceName(InCommand.WorkspaceName);
+		InCommand.bCommandSuccessful = PlasticSourceControlUtils::GetWorkspaceName(GetProvider().GetPathToWorkspaceRoot(), InCommand.WorkspaceName, InCommand.ErrorMessages);
 		if (InCommand.bCommandSuccessful)
 		{
-			// Get repository, server Url, branch and current changeset number
-			InCommand.bCommandSuccessful = PlasticSourceControlUtils::GetWorkspaceInformation(InCommand.ChangesetNumber, InCommand.RepositoryName, InCommand.ServerUrl, InCommand.BranchName);
+			// Get repository, server URL, branch and current changeset number
+			// Note: this initiates the connection to the server and issue network calls, so we don't need an explicit 'checkconnection'
+			InCommand.bCommandSuccessful = PlasticSourceControlUtils::GetWorkspaceInformation(InCommand.ChangesetNumber, InCommand.RepositoryName, InCommand.ServerUrl, InCommand.BranchName, InCommand.ErrorMessages);
 			if (InCommand.bCommandSuccessful)
 			{
-				// Execute a 'checkconnection' command to check the connectivity of the server.
-				InCommand.bCommandSuccessful = PlasticSourceControlUtils::RunCommand(TEXT("checkconnection"), TArray<FString>(), TArray<FString>(), InCommand.Concurrency, InCommand.InfoMessages, InCommand.ErrorMessages);
-				if (InCommand.bCommandSuccessful)
+				InCommand.InfoMessages.Add(TEXT("Connected successfully"));
+
+				// Now update the status of assets in the Content directory
+				// but only on real (re-)connection (but not each time Login() is called by Rename or Fixup Redirector command to check connection)
+				// and only if enabled in the settings
+				if (!GetProvider().IsAvailable() && GetProvider().AccessSettings().GetUpdateStatusAtStartup())
 				{
-					// Now update the status of assets in the Content directory
-					// but only on real (re-)connection (but not each time Login() is called by Rename or Fixup Redirector command to check connection)
-					// and only if enabled in the settings
-					if (!PlasticSourceControl.GetProvider().IsAvailable() && PlasticSourceControl.AccessSettings().GetUpdateStatusAtStartup())
-					{
-						TArray<FString> ContentDir;
-						ContentDir.Add(FPaths::ConvertRelativePathToFull(FPaths::ProjectContentDir()));
-						PlasticSourceControlUtils::RunUpdateStatus(ContentDir, false, InCommand.Concurrency, InCommand.ErrorMessages, States, InCommand.ChangesetNumber, InCommand.BranchName);
-					}
+					TArray<FString> ContentDir;
+					ContentDir.Add(FPaths::ConvertRelativePathToFull(FPaths::ProjectContentDir()));
+					PlasticSourceControlUtils::RunUpdateStatus(ContentDir, false, InCommand.Concurrency, InCommand.ErrorMessages, States, InCommand.ChangesetNumber, InCommand.BranchName);
 				}
-				else
+			}
+			else
+			{
+				if (InCommand.ErrorMessages.Num() > 0)
 				{
 					Operation->SetErrorText(FText::FromString(InCommand.ErrorMessages[0]));
 				}
+				else
+				{
+					const FText ErrorText(LOCTEXT("FailedToConnect", "Failed to connect to the Plastic SCM server."));
+					Operation->SetErrorText(ErrorText);
+					InCommand.ErrorMessages.Add(ErrorText.ToString());
+				}
 			}
 		}
 		else
 		{
-			Operation->SetErrorText(LOCTEXT("NotAPlasticRepository", "Failed to enable Plastic SCM source control. You need to initialize the project as a Plastic SCM repository first."));
+			const FText ErrorText(LOCTEXT("NotAPlasticRepository", "Failed to enable Plastic SCM source control. You need to create a Plastic SCM workspace for the project first."));
+			Operation->SetErrorText(ErrorText);
+			InCommand.ErrorMessages.Add(ErrorText.ToString());
 		}
 	}
 	else
 	{
-		Operation->SetErrorText(LOCTEXT("PlasticScmCliUnavaillable", "Failed to launch Plastic SCM command line tool. You need to install it or set the correct path to it first."));
+		const FText ErrorText(LOCTEXT("PlasticScmCliUnavaillable", "Failed to launch Plastic SCM 'cm' command line tool. You need to install it and make sure that 'cm' is on the Path and correctly configured."));
+		Operation->SetErrorText(ErrorText);
+		InCommand.ErrorMessages.Add(ErrorText.ToString());
 	}
 
 	return InCommand.bCommandSuccessful;
@@ -113,9 +167,36 @@ bool FPlasticConnectWorker::Execute(FPlasticSourceControlCommand& InCommand)
 
 bool FPlasticConnectWorker::UpdateStates()
 {
+	TRACE_CPUPROFILER_EVENT_SCOPE(FPlasticConnectWorker::UpdateStates);
+
 	return PlasticSourceControlUtils::UpdateCachedStates(MoveTemp(States));
 }
 
+
+static void UpdateChangelistState(FPlasticSourceControlProvider& SCCProvider, const FPlasticSourceControlChangelist& InChangelist, const TArray<FPlasticSourceControlState>& InStates)
+{
+	TRACE_CPUPROFILER_EVENT_SCOPE(FPlastic::UpdateChangelistState);
+
+	if (InChangelist.IsInitialized())
+	{
+		TSharedRef<FPlasticSourceControlChangelistState, ESPMode::ThreadSafe> ChangelistState = SCCProvider.GetStateInternal(InChangelist);
+
+		for (const FPlasticSourceControlState& InState : InStates)
+		{
+			if ((InState.WorkspaceState != EWorkspaceState::CheckedOut) && (InState.WorkspaceState != EWorkspaceState::Added) && InState.WorkspaceState != EWorkspaceState::Deleted)
+			{
+				continue;
+			}
+
+			TSharedRef<FPlasticSourceControlState, ESPMode::ThreadSafe> State = SCCProvider.GetStateInternal(InState.GetFilename());
+			ChangelistState->Files.Add(State);
+
+			// Keep the changelist stored with cached file state in sync with the actual changelist that owns this file.
+			State->Changelist = InChangelist;
+		}
+	}
+}
+
 FName FPlasticCheckOutWorker::GetName() const
 {
 	return "CheckOut";
@@ -123,6 +204,8 @@ FName FPlasticCheckOutWorker::GetName() const
 
 bool FPlasticCheckOutWorker::Execute(FPlasticSourceControlCommand& InCommand)
 {
+	TRACE_CPUPROFILER_EVENT_SCOPE(FPlasticCheckOutWorker::Execute);
+
 	check(InCommand.Operation->GetName() == GetName());
 
 	// Detect special case for a partial checkout (CS:-1 in Gluon mode)!
@@ -143,10 +226,47 @@ bool FPlasticCheckOutWorker::Execute(FPlasticSourceControlCommand& InCommand)
 
 bool FPlasticCheckOutWorker::UpdateStates()
 {
+	TRACE_CPUPROFILER_EVENT_SCOPE(FPlasticCheckOutWorker::UpdateStates);
+
+	// If files have been checked-out directly to a CL, modify the cached state to reflect it (defaults to the Default changelist).
+	UpdateChangelistState(GetProvider(), InChangelist, States);
+
 	return PlasticSourceControlUtils::UpdateCachedStates(MoveTemp(States));
 }
 
-/// Parse check-in result, usually locking like "Created changeset cs:8@br:/main@MyProject@SRombauts@cloud (mount:'/')"
+bool DeleteChangelist(const FPlasticSourceControlProvider& PlasticSourceControlProvider, const FPlasticSourceControlChangelist& InChangelist, const EConcurrency::Type InConcurrency, TArray<FString>& OutResults, TArray<FString>& OutErrorMessages)
+{
+	TArray<FString> Parameters;
+	TArray<FString> Files;
+	if (PlasticSourceControlProvider.GetPlasticScmVersion() < s_NewChangelistFileArgsPlasticScmVersion)
+	{
+		Parameters.Add(TEXT("rm"));
+		Files.Add(InChangelist.GetName());
+		return PlasticSourceControlUtils::RunCommand(TEXT("changelist"), Parameters, Files, InConcurrency, OutResults, OutErrorMessages);
+	}
+	else
+	{
+		Parameters.Add(TEXT("delete"));
+		const FScopedTempFile ChangelistNameFile(InChangelist.GetName());
+		Parameters.Add(FString::Printf(TEXT("--namefile=\"%s\""), *FPaths::ConvertRelativePathToFull(ChangelistNameFile.GetFilename())));
+		UE_LOG(LogSourceControl, Verbose, TEXT("DeleteChangelist(%s)"), *InChangelist.GetName());
+		return PlasticSourceControlUtils::RunCommand(TEXT("changelist"), Parameters, Files, InConcurrency, OutResults, OutErrorMessages);
+	}
+}
+
+TArray<FString> FileNamesFromFileStates(const TArray<FSourceControlStateRef>& InFileStates)
+{
+	TArray<FString> Files;
+
+	for (const FSourceControlStateRef& FileState : InFileStates)
+	{
+		Files.Add(FileState->GetFilename());
+	}
+
+	return Files;
+}
+
+/// Parse checkin result, usually looking like "Created changeset cs:8@br:/main@MyProject@SRombauts@cloud (mount:'/')"
 static FText ParseCheckInResults(const TArray<FString>& InResults)
 {
 	if (InResults.Num() > 0)
@@ -178,62 +298,100 @@ FName FPlasticCheckInWorker::GetName() const
 
 bool FPlasticCheckInWorker::Execute(FPlasticSourceControlCommand& InCommand)
 {
+	TRACE_CPUPROFILER_EVENT_SCOPE(FPlasticCheckInWorker::Execute);
+
 	check(InCommand.Operation->GetName() == GetName());
 	TSharedRef<FCheckIn, ESPMode::ThreadSafe> Operation = StaticCastSharedRef<FCheckIn>(InCommand.Operation);
 
-	UE_LOG(LogSourceControl, Verbose, TEXT("CheckIn: %d file(s) Description: '%s'"), InCommand.Files.Num(), *Operation->GetDescription().ToString());
+	TArray<FString> Files;
+	if (InCommand.Changelist.IsInitialized() && InCommand.Files.IsEmpty())
+	{
+		TSharedRef<FPlasticSourceControlChangelistState, ESPMode::ThreadSafe> ChangelistState = GetProvider().GetStateInternal(InCommand.Changelist);
+		Files = FileNamesFromFileStates(ChangelistState->Files);
 
-	// make a temp file to place our commit message in
-	FScopedTempFile CommitMsgFile(Operation->GetDescription());
-	if (!CommitMsgFile.GetFilename().IsEmpty())
+		InChangelist = InCommand.Changelist;
+	}
+	else
 	{
-		TArray<FString> Parameters;
-		FString ParamCommitMsgFilename = TEXT("--commentsfile=\"");
-		ParamCommitMsgFilename += FPaths::ConvertRelativePathToFull(CommitMsgFile.GetFilename());
-		ParamCommitMsgFilename += TEXT("\"");
-		Parameters.Add(ParamCommitMsgFilename);
-		// Detect special case for a partial checkout (CS:-1 in Gluon mode)!
-		if (-1 != InCommand.ChangesetNumber)
-		{
-			Parameters.Add(TEXT("--all")); // Also files Changed (not CheckedOut) and Moved/Deleted Locally
-		//  NOTE: --update added as #23 but removed as #32 because most assets are locked by the Unreal Editor
-		//	Parameters.Add(TEXT("--update")); // Processes the update-merge automatically if it eventually happens.
-			InCommand.bCommandSuccessful = PlasticSourceControlUtils::RunCommand(TEXT("checkin"), Parameters, InCommand.Files, InCommand.Concurrency, InCommand.InfoMessages, InCommand.ErrorMessages);
-		}
-		else
-		{
-			Parameters.Add(TEXT("--applychanged")); // Also files Changed (not CheckedOut) and Moved/Deleted Locally
-			InCommand.bCommandSuccessful = PlasticSourceControlUtils::RunCommand(TEXT("partial checkin"), Parameters, InCommand.Files, InCommand.Concurrency, InCommand.InfoMessages, InCommand.ErrorMessages);
-		}
-		if (InCommand.bCommandSuccessful)
-		{
-			// Remove any deleted files from status cache
-			FPlasticSourceControlModule& PlasticSourceControl = FModuleManager::GetModuleChecked<FPlasticSourceControlModule>("PlasticSourceControl");
-			FPlasticSourceControlProvider& Provider = PlasticSourceControl.GetProvider();
+		Files = InCommand.Files;
+	}
 
-			TArray<TSharedRef<ISourceControlState, ESPMode::ThreadSafe>> LocalStates;
-			Provider.GetState(InCommand.Files, LocalStates, EStateCacheUsage::Use);
-			for (const auto& State : LocalStates)
+	if (Files.Num() > 0)
+	{
+		UE_LOG(LogSourceControl, Verbose, TEXT("CheckIn: %d file(s) Description: '%s'"), Files.Num(), *Operation->GetDescription().ToString());
+
+		// make a temp file to place our commit message in
+		const FScopedTempFile CommitMsgFile(Operation->GetDescription());
+		if (!CommitMsgFile.GetFilename().IsEmpty())
+		{
+			TArray<FString> Parameters;
+			Parameters.Add(FString::Printf(TEXT("--commentsfile=\"%s\""), *FPaths::ConvertRelativePathToFull(CommitMsgFile.GetFilename())));
+			// Detect special case for a partial checkout (CS:-1 in Gluon mode)!
+			if (-1 != InCommand.ChangesetNumber)
 			{
-				if (State->IsDeleted())
+				Parameters.Add(TEXT("--all")); // Also files Changed (not CheckedOut) and Moved/Deleted Locally
+			//  NOTE: --update added as #23 but removed as #32 because most assets are locked by the Unreal Editor
+			//  Parameters.Add(TEXT("--update")); // Processes the update-merge automatically if it eventually happens.
+				InCommand.bCommandSuccessful = PlasticSourceControlUtils::RunCommand(TEXT("checkin"), Parameters, Files, InCommand.Concurrency, InCommand.InfoMessages, InCommand.ErrorMessages);
+			}
+			else
+			{
+				Parameters.Add(TEXT("--applychanged")); // Also files Changed (not CheckedOut) and Moved/Deleted Locally
+				InCommand.bCommandSuccessful = PlasticSourceControlUtils::RunCommand(TEXT("partial checkin"), Parameters, Files, InCommand.Concurrency, InCommand.InfoMessages, InCommand.ErrorMessages);
+			}
+			if (InCommand.bCommandSuccessful)
+			{
+				// Remove any deleted files from status cache
+				TArray<TSharedRef<ISourceControlState, ESPMode::ThreadSafe>> LocalStates;
+				GetProvider().GetState(Files, LocalStates, EStateCacheUsage::Use);
+				for (const auto& State : LocalStates)
 				{
-					Provider.RemoveFileFromCache(State->GetFilename());
+					if (State->IsDeleted())
+					{
+						GetProvider().RemoveFileFromCache(State->GetFilename());
+					}
 				}
+
+				Operation->SetSuccessMessage(ParseCheckInResults(InCommand.InfoMessages));
+				UE_LOG(LogSourceControl, Log, TEXT("CheckIn successful"));
 			}
 
-			Operation->SetSuccessMessage(ParseCheckInResults(InCommand.InfoMessages));
-			UE_LOG(LogSourceControl, Log, TEXT("CheckIn successful"));
+			if (InChangelist.IsInitialized() && !InChangelist.IsDefault())
+			{
+				// NOTE: we need to explicitly delete persistent changelists when we submit its content
+				DeleteChangelist(GetProvider(), InChangelist, InCommand.Concurrency, InCommand.InfoMessages, InCommand.ErrorMessages);
+			}
 		}
-	}
 
-	// now update the status of our files
-	PlasticSourceControlUtils::RunUpdateStatus(InCommand.Files, false, InCommand.Concurrency, InCommand.ErrorMessages, States, InCommand.ChangesetNumber, InCommand.BranchName);
+		// now update the status of our files
+		PlasticSourceControlUtils::RunUpdateStatus(Files, false, InCommand.Concurrency, InCommand.ErrorMessages, States, InCommand.ChangesetNumber, InCommand.BranchName);
+	}
+	else
+	{
+		UE_LOG(LogSourceControl, Warning, TEXT("Checkin: No files provided"));
+	}
 
 	return InCommand.bCommandSuccessful;
 }
 
 bool FPlasticCheckInWorker::UpdateStates()
 {
+	TRACE_CPUPROFILER_EVENT_SCOPE(FPlasticCheckInWorker::UpdateStates);
+
+	if (InChangelist.IsInitialized())
+	{
+		if (InChangelist.IsDefault())
+		{
+			// Remove all the files from the default changelist state, since they have been submitted, but we didn't delete the changelist itself
+			TSharedRef<FPlasticSourceControlChangelistState, ESPMode::ThreadSafe> DefaultChangelist = GetProvider().GetStateInternal(FPlasticSourceControlChangelist::DefaultChangelist);
+			DefaultChangelist->Files.Empty();
+		}
+		else
+		{
+			GetProvider().RemoveChangelistFromCache(InChangelist);
+		}
+	}
+
 	return PlasticSourceControlUtils::UpdateCachedStates(MoveTemp(States));
 }
 
@@ -244,42 +402,56 @@ FName FPlasticMarkForAddWorker::GetName() const
 
 bool FPlasticMarkForAddWorker::Execute(FPlasticSourceControlCommand& InCommand)
 {
+	TRACE_CPUPROFILER_EVENT_SCOPE(FPlasticMarkForAddWorker::Execute);
+
 	check(InCommand.Operation->GetName() == GetName());
 
-	TArray<FString> Parameters;
-	Parameters.Add(TEXT("--parents")); // NOTE: deprecated in 8.0.16.3100 when it became the default https://www.plasticscm.com/download/releasenotes/8.0.16.3100
-	// Note: using "?" is a workaround to trigger the Plastic's "SkipIgnored" internal flag meaning "don't add file that are ignored":
-	//			options.SkipIgnored = cla.GetWildCardArguments().Count > 0;
-	//		 It's behavior is similar as Subversion:
-	//  		if you explicitely add one file that is ignored, "cm" will happily accept it and add it,
-	//			if you try to add a set of files with a pattern, "cm" will skip the files that are ignored and only add the other ones
-	// TODO: provide an updated version of "cm" with a new flag like --applyignorerules
-	if (AreAllFiles(InCommand.Files))
-	{
-		Parameters.Add(TEXT("?"));	// needed only when used with a list of files
-	}
-	else
-	{
-		Parameters.Add(TEXT("-R"));	// needed only at the time of workspace creation, to add directories recursively
-	}
-	// Detect special case for a partial checkout (CS:-1 in Gluon mode)!
-	if (-1 != InCommand.ChangesetNumber)
+	if (InCommand.Files.Num() > 0)
 	{
-		InCommand.bCommandSuccessful = PlasticSourceControlUtils::RunCommand(TEXT("add"), Parameters, InCommand.Files, InCommand.Concurrency, InCommand.InfoMessages, InCommand.ErrorMessages);
+		TArray<FString> Parameters;
+		Parameters.Add(TEXT("--parents")); // NOTE: deprecated in 8.0.16.3100 when it became the default https://www.plasticscm.com/download/releasenotes/8.0.16.3100
+		// Note: using "?" is a workaround to trigger the Plastic's "SkipIgnored" internal flag meaning "don't add file that are ignored":
+		//          options.SkipIgnored = cla.GetWildCardArguments().Count > 0;
+		//       It's behavior is similar as Subversion:
+		//          if you explicitly add one file that is ignored, "cm" will happily accept it and add it,
+		//          if you try to add a set of files with a pattern, "cm" will skip the files that are ignored and only add the other ones
+		// TODO: provide an updated version of "cm" with a new flag like --applyignorerules
+		if (AreAllFiles(InCommand.Files))
+		{
+			Parameters.Add(TEXT("?"));	// needed only when used with a list of files
+		}
+		else
+		{
+			Parameters.Add(TEXT("-R"));	// needed only at the time of workspace creation, to add directories recursively
+		}
+		// Detect special case for a partial checkout (CS:-1 in Gluon mode)!
+		if (-1 != InCommand.ChangesetNumber)
+		{
+			InCommand.bCommandSuccessful = PlasticSourceControlUtils::RunCommand(TEXT("add"), Parameters, InCommand.Files, InCommand.Concurrency, InCommand.InfoMessages, InCommand.ErrorMessages);
+		}
+		else
+		{
+			InCommand.bCommandSuccessful = PlasticSourceControlUtils::RunCommand(TEXT("partial add"), Parameters, InCommand.Files, InCommand.Concurrency, InCommand.InfoMessages, InCommand.ErrorMessages);
+		}
+
+		// now update the status of our files
+		PlasticSourceControlUtils::RunUpdateStatus(InCommand.Files, false, InCommand.Concurrency, InCommand.ErrorMessages, States, InCommand.ChangesetNumber, InCommand.BranchName);
 	}
 	else
 	{
-		InCommand.bCommandSuccessful = PlasticSourceControlUtils::RunCommand(TEXT("partial add"), Parameters, InCommand.Files, InCommand.Concurrency, InCommand.InfoMessages, InCommand.ErrorMessages);
+		UE_LOG(LogSourceControl, Warning, TEXT("MarkforAdd: No files provided"));
 	}
 
-	// now update the status of our files
-	PlasticSourceControlUtils::RunUpdateStatus(InCommand.Files, false, InCommand.Concurrency, InCommand.ErrorMessages, States, InCommand.ChangesetNumber, InCommand.BranchName);
-
 	return InCommand.bCommandSuccessful;
 }
 
 bool FPlasticMarkForAddWorker::UpdateStates()
 {
+	TRACE_CPUPROFILER_EVENT_SCOPE(FPlasticMarkForAddWorker::UpdateStates);
+
+	// If files have been added directly to a CL, modify the cached state to reflect it (defaults to the Default changelist).
+	UpdateChangelistState(GetProvider(), InChangelist, States);
+
 	return PlasticSourceControlUtils::UpdateCachedStates(MoveTemp(States));
 }
 
@@ -290,6 +462,8 @@ FName FPlasticDeleteWorker::GetName() const
 
 bool FPlasticDeleteWorker::Execute(FPlasticSourceControlCommand& InCommand)
 {
+	TRACE_CPUPROFILER_EVENT_SCOPE(FPlasticDeleteWorker::Execute);
+
 	check(InCommand.Operation->GetName() == GetName());
 
 	// Detect special case for a partial checkout (CS:-1 in Gluon mode)!
@@ -310,6 +484,11 @@ bool FPlasticDeleteWorker::Execute(FPlasticSourceControlCommand& InCommand)
 
 bool FPlasticDeleteWorker::UpdateStates()
 {
+	TRACE_CPUPROFILER_EVENT_SCOPE(FPlasticDeleteWorkers::UpdateStates);
+
+	// If files have been deleted directly to a CL, modify the cached state to reflect it (defaults to the Default changelist).
+	UpdateChangelistState(GetProvider(), InChangelist, States);
+
 	return PlasticSourceControlUtils::UpdateCachedStates(MoveTemp(States));
 }
 
@@ -320,17 +499,27 @@ FName FPlasticRevertWorker::GetName() const
 
 bool FPlasticRevertWorker::Execute(FPlasticSourceControlCommand& InCommand)
 {
+	TRACE_CPUPROFILER_EVENT_SCOPE(FPlasticRevertWorker::Execute);
+
 	check(InCommand.Operation->GetName() == GetName());
 
-	FPlasticSourceControlModule& PlasticSourceControl = FModuleManager::GetModuleChecked<FPlasticSourceControlModule>("PlasticSourceControl");
-	FPlasticSourceControlProvider& Provider = PlasticSourceControl.GetProvider();
+	TArray<FString> Files;
+	if (InCommand.Changelist.IsInitialized() && InCommand.Files.IsEmpty())
+	{
+		TSharedRef<FPlasticSourceControlChangelistState, ESPMode::ThreadSafe> ChangelistState = GetProvider().GetStateInternal(InCommand.Changelist);
+		Files = FileNamesFromFileStates(ChangelistState->Files);
+	}
+	else
+	{
+		Files = InCommand.Files;
+	}
 
 	TArray<FString> ChangedFiles;
 	TArray<FString> CheckedOutFiles;
 
-	for (const FString& File : InCommand.Files)
+	for (const FString& File : Files)
 	{
-		TSharedRef<FPlasticSourceControlState, ESPMode::ThreadSafe> State = Provider.GetStateInternal(File);
+		TSharedRef<FPlasticSourceControlState, ESPMode::ThreadSafe> State = GetProvider().GetStateInternal(File);
 
 		if (EWorkspaceState::Changed == State->WorkspaceState)
 		{
@@ -372,14 +561,29 @@ bool FPlasticRevertWorker::Execute(FPlasticSourceControlCommand& InCommand)
 		}
 	}
 
-	// NOTE: optim, no need to update the status of our files since this is done immediately after by the Editor
+	// update the status of our files
+	PlasticSourceControlUtils::RunUpdateStatus(Files, false, InCommand.Concurrency, InCommand.ErrorMessages, States, InCommand.ChangesetNumber, InCommand.BranchName);
 
 	return InCommand.bCommandSuccessful;
 }
 
 bool FPlasticRevertWorker::UpdateStates()
 {
-	return false;
+	TRACE_CPUPROFILER_EVENT_SCOPE(FPlasticRevertWorker::UpdateStates);
+
+	// Update affected changelist if any
+	for (const FPlasticSourceControlState& NewState : States)
+	{
+		TSharedRef<FPlasticSourceControlState, ESPMode::ThreadSafe> State = GetProvider().GetStateInternal(NewState.GetFilename());
+		if (State->Changelist.IsInitialized())
+		{
+			// 1- Remove these files from their previous changelist
+			TSharedRef<FPlasticSourceControlChangelistState, ESPMode::ThreadSafe> PreviousChangelist = GetProvider().GetStateInternal(State->Changelist);
+			PreviousChangelist->Files.Remove(State);
+		}
+	}
+
+	return PlasticSourceControlUtils::UpdateCachedStates(MoveTemp(States));
 }
 
 FName FPlasticRevertUnchangedWorker::GetName() const
@@ -389,24 +593,56 @@ FName FPlasticRevertUnchangedWorker::GetName() const
 
 bool FPlasticRevertUnchangedWorker::Execute(FPlasticSourceControlCommand& InCommand)
 {
+	TRACE_CPUPROFILER_EVENT_SCOPE(FPlasticRevertUnchangedWorker::Execute);
+
 	check(InCommand.Operation->GetName() == GetName());
 
 	TArray<FString> Parameters;
 	Parameters.Add(TEXT("-R"));
 
+	TArray<FString> Files;
+	if (InCommand.Changelist.IsInitialized() && InCommand.Files.IsEmpty())
+	{
+		TSharedRef<FPlasticSourceControlChangelistState, ESPMode::ThreadSafe> ChangelistState = GetProvider().GetStateInternal(InCommand.Changelist);
+		Files = FileNamesFromFileStates(ChangelistState->Files);
+	}
+	else
+	{
+		Files = InCommand.Files;
+	}
+
 	// revert the checkout of all unchanged files recursively
-	InCommand.bCommandSuccessful = PlasticSourceControlUtils::RunCommand(TEXT("uncounchanged"), Parameters, InCommand.Files, InCommand.Concurrency, InCommand.InfoMessages, InCommand.ErrorMessages);
+	InCommand.bCommandSuccessful = PlasticSourceControlUtils::RunCommand(TEXT("uncounchanged"), Parameters, Files, InCommand.Concurrency, InCommand.InfoMessages, InCommand.ErrorMessages);
 
-	// Now update the status of assets in the Content directory
-	TArray<FString> ContentDir;
-	ContentDir.Add(FPaths::ConvertRelativePathToFull(FPaths::ProjectContentDir()));
-	PlasticSourceControlUtils::RunUpdateStatus(ContentDir, false, InCommand.Concurrency, InCommand.ErrorMessages, States, InCommand.ChangesetNumber, InCommand.BranchName);
+	// Now update the status of either the files, or all assets in the Content directory
+	if (Files.Num() == 0)
+	{
+		Files.Add(FPaths::ConvertRelativePathToFull(FPaths::ProjectContentDir()));
+	}
+	PlasticSourceControlUtils::RunUpdateStatus(Files, false, InCommand.Concurrency, InCommand.ErrorMessages, States, InCommand.ChangesetNumber, InCommand.BranchName);
 
 	return InCommand.bCommandSuccessful;
 }
 
 bool FPlasticRevertUnchangedWorker::UpdateStates()
 {
+	TRACE_CPUPROFILER_EVENT_SCOPE(FPlasticRevertUnchangedWorker::UpdateStates);
+
+	// Update affected changelist if any
+	for (const FPlasticSourceControlState& NewState : States)
+	{
+		if (!NewState.IsModified())
+		{
+			TSharedRef<FPlasticSourceControlState, ESPMode::ThreadSafe> State = GetProvider().GetStateInternal(NewState.GetFilename());
+			if (State->Changelist.IsInitialized())
+			{
+				// 1- Remove these files from their previous changelist
+				TSharedRef<FPlasticSourceControlChangelistState, ESPMode::ThreadSafe> PreviousChangelist = GetProvider().GetStateInternal(State->Changelist);
+				PreviousChangelist->Files.Remove(State);
+			}
+		}
+	}
+
 	return PlasticSourceControlUtils::UpdateCachedStates(MoveTemp(States));
 }
 
@@ -417,6 +653,8 @@ FName FPlasticRevertAllWorker::GetName() const
 
 bool FPlasticRevertAllWorker::Execute(FPlasticSourceControlCommand& InCommand)
 {
+	TRACE_CPUPROFILER_EVENT_SCOPE(FPlasticRevertAllWorker::Execute);
+
 	check(InCommand.Operation->GetName() == GetName());
 
 	TArray<FString> Parameters;
@@ -442,6 +680,8 @@ bool FPlasticRevertAllWorker::Execute(FPlasticSourceControlCommand& InCommand)
 
 bool FPlasticRevertAllWorker::UpdateStates()
 {
+	TRACE_CPUPROFILER_EVENT_SCOPE(FPlasticRevertAllWorker::UpdateStates);
+
 	return PlasticSourceControlUtils::UpdateCachedStates(MoveTemp(States));
 }
 
@@ -452,6 +692,8 @@ FName FPlasticMakeWorkspaceWorker::GetName() const
 
 bool FPlasticMakeWorkspaceWorker::Execute(FPlasticSourceControlCommand& InCommand)
 {
+	TRACE_CPUPROFILER_EVENT_SCOPE(FPlasticMakeWorkspaceWorker::Execute);
+
 	check(InCommand.Operation->GetName() == GetName());
 	TSharedRef<FPlasticMakeWorkspace, ESPMode::ThreadSafe> Operation = StaticCastSharedRef<FPlasticMakeWorkspace>(InCommand.Operation);
 
@@ -484,6 +726,8 @@ FName FPlasticUpdateStatusWorker::GetName() const
 
 bool FPlasticUpdateStatusWorker::Execute(FPlasticSourceControlCommand& InCommand)
 {
+	TRACE_CPUPROFILER_EVENT_SCOPE(FPlasticUpdateStatusWorker::Execute);
+
 	check(InCommand.Operation->GetName() == GetName());
 	TSharedRef<FUpdateStatus, ESPMode::ThreadSafe> Operation = StaticCastSharedRef<FUpdateStatus>(InCommand.Operation);
 
@@ -491,9 +735,20 @@ bool FPlasticUpdateStatusWorker::Execute(FPlasticSourceControlCommand& InCommand
 	UE_LOG(LogSourceControl, Log, TEXT("status (of %d files, ShouldUpdateHistory=%d, ShouldGetOpenedOnly=%d, ShouldUpdateModifiedState=%d)"),
 		InCommand.Files.Num(), Operation->ShouldUpdateHistory(), Operation->ShouldGetOpenedOnly(), Operation->ShouldUpdateModifiedState());
 
-	if (InCommand.Files.Num() > 0)
+	TArray<FString> Files;
+	if (InCommand.Changelist.IsInitialized() && InCommand.Files.IsEmpty())
+	{
+		TSharedRef<FPlasticSourceControlChangelistState, ESPMode::ThreadSafe> ChangelistState = GetProvider().GetStateInternal(InCommand.Changelist);
+		Files = FileNamesFromFileStates(ChangelistState->Files);
+	}
+	else
 	{
-		InCommand.bCommandSuccessful = PlasticSourceControlUtils::RunUpdateStatus(InCommand.Files, Operation->ShouldUpdateHistory(), InCommand.Concurrency, InCommand.ErrorMessages, States, InCommand.ChangesetNumber, InCommand.BranchName);
+		Files = InCommand.Files;
+	}
+
+	if (Files.Num() > 0)
+	{
+		InCommand.bCommandSuccessful = PlasticSourceControlUtils::RunUpdateStatus(Files, Operation->ShouldUpdateHistory(), InCommand.Concurrency, InCommand.ErrorMessages, States, InCommand.ChangesetNumber, InCommand.BranchName);
 		// Remove all "is not in a workspace" error and convert the result to "success" if there are no other errors
 		PlasticSourceControlUtils::RemoveRedundantErrors(InCommand, TEXT("is not in a workspace."));
 		if (!InCommand.bCommandSuccessful)
@@ -535,8 +790,8 @@ bool FPlasticUpdateStatusWorker::Execute(FPlasticSourceControlCommand& InCommand
 		}
 		else
 		{
-			const FPlasticSourceControlModule& PlasticSourceControl = FModuleManager::GetModuleChecked<FPlasticSourceControlModule>("PlasticSourceControl");
-			if (PlasticSourceControl.AccessSettings().GetUpdateStatusOtherBranches() && AreAllFiles(InCommand.Files))
+			FPlasticSourceControlSettings& PlasticSettings = GetProvider().AccessSettings();
+			if (PlasticSettings.GetUpdateStatusOtherBranches() && AreAllFiles(Files))
 			{
 				// Get only the last revision of the files (checking all branches)
 				// in order to warn the user if the file has been changed on another branch
@@ -553,6 +808,12 @@ bool FPlasticUpdateStatusWorker::Execute(FPlasticSourceControlCommand& InCommand
 		ProjectDirs.Add(FPaths::ConvertRelativePathToFull(FPaths::ProjectContentDir()));
 		InCommand.bCommandSuccessful = PlasticSourceControlUtils::RunUpdateStatus(ProjectDirs, Operation->ShouldUpdateHistory(), InCommand.Concurrency, InCommand.ErrorMessages, States, InCommand.ChangesetNumber, InCommand.BranchName);
 	}
+	else
+	{
+		// TODO: workaround for the case of submitting a changelist, calling UpdateStatus with no files nor the changelist.
+		// No consequences, and no way to fix it, so let's not show an error.
+		InCommand.bCommandSuccessful = true;
+	}
 
 	// TODO: re-evaluate how to optimize this heavy operation using some of these hints flags
 	// - ShouldGetOpenedOnly hint would be to call for all a whole workspace status update
@@ -563,49 +824,9 @@ bool FPlasticUpdateStatusWorker::Execute(FPlasticSourceControlCommand& InCommand
 
 bool FPlasticUpdateStatusWorker::UpdateStates()
 {
-	return PlasticSourceControlUtils::UpdateCachedStates(MoveTemp(States));
-}
+	TRACE_CPUPROFILER_EVENT_SCOPE(FPlasticUpdateStatusWorker::UpdateStates);
 
-/// Detect if the operation is a duplicate/copy or a rename/move, and if it leaved a redirector (ie it was a move of a source controled asset)
-bool IsMoveOperation(const FString& InOrigin)
-{
-	bool bIsMoveOperation = true;
-
-	FString PackageName;
-	if (FPackageName::TryConvertFilenameToLongPackageName(InOrigin, PackageName))
-	{
-		// Use AsyncTask to call AssetRegistry GetAssetsByPackageName') on Game Thread
-		const TSharedRef<TPromise<TArray<FAssetData>>, ESPMode::ThreadSafe> Promise = MakeShareable(new TPromise<TArray<FAssetData>>());
-		AsyncTask(ENamedThreads::GameThread, [Promise, PackageName]()
-		{
-			TArray<FAssetData> AssetsData;
-			FAssetRegistryModule& AssetRegistryModule = FModuleManager::GetModuleChecked<FAssetRegistryModule>(TEXT("AssetRegistry"));
-			AssetRegistryModule.Get().GetAssetsByPackageName(FName(*PackageName), AssetsData);
-			Promise->SetValue(MoveTemp(AssetsData));
-		});
-		const TArray<FAssetData> AssetsData = Promise->GetFuture().Get();
-		UE_LOG(LogSourceControl, Log, TEXT("IsMoveOperation: PackageName: %s, AssetsData: Num=%d"), *PackageName, AssetsData.Num());
-		if (AssetsData.Num() > 0)
-		{
-			const FAssetData& AssetData = AssetsData[0];
-			if (!AssetData.IsRedirector())
-			{
-				UE_LOG(LogSourceControl, Log, TEXT("IsMoveOperation: %s is a plain asset, so it's a duplicate/copy"), *InOrigin);
-				bIsMoveOperation = false;
-			}
-			else
-			{
-				UE_LOG(LogSourceControl, Log, TEXT("IsMoveOperation: %s is a redirector, so it's a move/rename"), *InOrigin);
-			}
-		}
-		else
-		{
-			// no asset in package (no redirector) so it should be a rename/move of a newly Added (not Controlled/Checked-In) file
-			UE_LOG(LogSourceControl, Log, TEXT("IsMoveOperation: %s does not have asset in package (ie. no redirector) so it's a move/rename of a newly added file"), *InOrigin);
-		}
-	}
-
-	return bIsMoveOperation;
+	return PlasticSourceControlUtils::UpdateCachedStates(MoveTemp(States));
 }
 
 FName FPlasticCopyWorker::GetName() const
@@ -615,6 +836,8 @@ FName FPlasticCopyWorker::GetName() const
 
 bool FPlasticCopyWorker::Execute(FPlasticSourceControlCommand& InCommand)
 {
+	TRACE_CPUPROFILER_EVENT_SCOPE(FPlasticCopyWorker::Execute);
+
 	check(InCommand.Operation->GetName() == GetName());
 	TSharedRef<FCopy, ESPMode::ThreadSafe> Operation = StaticCastSharedRef<FCopy>(InCommand.Operation);
 
@@ -623,8 +846,8 @@ bool FPlasticCopyWorker::Execute(FPlasticSourceControlCommand& InCommand)
 		const FString& Origin = InCommand.Files[0];
 		const FString Destination = FPaths::ConvertRelativePathToFull(Operation->GetDestination());
 
-		// Detect if the operation is a duplicate/copy or a rename/move, and if it leaved a redirector (ie it was a move of a source controled asset)
-		const bool bIsMoveOperation = IsMoveOperation(Origin);
+		// Detects if the operation is a Move/Rename "Branch" that we need to track (or just a Duplicate/Copy, already Added to source control)
+		const bool bIsMoveOperation = (Operation->CopyMethod == FCopy::ECopyMethod::Branch);
 		if (bIsMoveOperation)
 		{
 			UE_LOG(LogSourceControl, Log, TEXT("Moving %s to %s..."), *Origin, *Destination);
@@ -691,6 +914,8 @@ bool FPlasticCopyWorker::Execute(FPlasticSourceControlCommand& InCommand)
 
 bool FPlasticCopyWorker::UpdateStates()
 {
+	TRACE_CPUPROFILER_EVENT_SCOPE(FPlasticCopyWorkers::UpdateStates);
+
 	return PlasticSourceControlUtils::UpdateCachedStates(MoveTemp(States));
 }
 
@@ -701,6 +926,8 @@ FName FPlasticSyncWorker::GetName() const
 
 bool FPlasticSyncWorker::Execute(FPlasticSourceControlCommand& InCommand)
 {
+	TRACE_CPUPROFILER_EVENT_SCOPE(FPlasticSyncWorker::Execute);
+
 	check(InCommand.Operation->GetName() == GetName());
 
 	TArray<FString> Parameters;
@@ -736,6 +963,8 @@ bool FPlasticSyncWorker::Execute(FPlasticSourceControlCommand& InCommand)
 
 bool FPlasticSyncWorker::UpdateStates()
 {
+	TRACE_CPUPROFILER_EVENT_SCOPE(FPlasticSyncWorker::UpdateStates);
+
 	return PlasticSourceControlUtils::UpdateCachedStates(MoveTemp(States));
 }
 
@@ -746,14 +975,14 @@ FName FPlasticResolveWorker::GetName() const
 
 bool FPlasticResolveWorker::Execute(FPlasticSourceControlCommand& InCommand)
 {
+	TRACE_CPUPROFILER_EVENT_SCOPE(FPlasticResolveWorker::Execute);
+
 	check(InCommand.Operation->GetName() == GetName());
 
 	// Currently resolve operation is always on one file only, but the following would works for many
 	for (const FString& File : InCommand.Files)
 	{
-		FPlasticSourceControlModule& PlasticSourceControl = FModuleManager::GetModuleChecked<FPlasticSourceControlModule>("PlasticSourceControl");
-		FPlasticSourceControlProvider& Provider = PlasticSourceControl.GetProvider();
-		TSharedRef<FPlasticSourceControlState, ESPMode::ThreadSafe> State = Provider.GetStateInternal(File);
+		auto State = GetProvider().GetStateInternal(File);
 
 		// To resolve the conflict, merge the file by keeping it like it is on file system
 		// TODO: according to documentation, this cannot work for cherry-picking
@@ -781,7 +1010,486 @@ bool FPlasticResolveWorker::Execute(FPlasticSourceControlCommand& InCommand)
 
 bool FPlasticResolveWorker::UpdateStates()
 {
+	TRACE_CPUPROFILER_EVENT_SCOPE(FPlasticResolveWorker::UpdateStates);
+
 	return PlasticSourceControlUtils::UpdateCachedStates(MoveTemp(States));
 }
 
+FName FPlasticGetPendingChangelistsWorker::GetName() const
+{
+	return "UpdateChangelistsStatus";
+}
+
+bool FPlasticGetPendingChangelistsWorker::Execute(FPlasticSourceControlCommand& InCommand)
+{
+	TRACE_CPUPROFILER_EVENT_SCOPE(FPlasticGetPendingChangelistsWorker::Execute);
+
+	check(InCommand.Operation->GetName() == GetName());
+	TSharedRef<FUpdatePendingChangelistsStatus, ESPMode::ThreadSafe> Operation = StaticCastSharedRef<FUpdatePendingChangelistsStatus>(InCommand.Operation);
+
+	InCommand.bCommandSuccessful = PlasticSourceControlUtils::RunGetChangelists(InCommand.Concurrency, OutChangelistsStates, OutCLFilesStates, InCommand.ErrorMessages);
+	if (InCommand.bCommandSuccessful)
+	{
+		// Remove the changelist that were not requested by the user.
+		if (!Operation->ShouldUpdateAllChangelists())
+		{
+			const TArray<FSourceControlChangelistRef>& RequestedChangelists = Operation->GetChangelistsToUpdate();
+			OutChangelistsStates.RemoveAll([&RequestedChangelists](const FPlasticSourceControlChangelistState& ChangelistState)
+				{
+					FPlasticSourceControlChangelistRef RemoveChangelistCandidate = StaticCastSharedRef<FPlasticSourceControlChangelist>(ChangelistState.GetChangelist());
+					return !RequestedChangelists.ContainsByPredicate([&RemoveChangelistCandidate](const FSourceControlChangelistRef& Requested)
+						{
+							return StaticCastSharedRef<FPlasticSourceControlChangelist>(Requested)->GetName() == RemoveChangelistCandidate->GetName();
+						});
+				});
+		}
+	}
+
+	bCleanupCache = InCommand.bCommandSuccessful;
+
+	return InCommand.bCommandSuccessful;
+}
+
+bool FPlasticGetPendingChangelistsWorker::UpdateStates()
+{
+	TRACE_CPUPROFILER_EVENT_SCOPE(FPlasticGetPendingChangelistsWorker::UpdateStates);
+
+	bool bUpdated = false;
+
+	const FDateTime Now = FDateTime::Now();
+
+	// first update cached state from 'changes' call
+	for (int StatusIndex = 0; StatusIndex < OutChangelistsStates.Num(); StatusIndex++)
+	{
+		const FPlasticSourceControlChangelistState& CLStatus = OutChangelistsStates[StatusIndex];
+		TSharedRef<FPlasticSourceControlChangelistState, ESPMode::ThreadSafe> ChangelistState = GetProvider().GetStateInternal(CLStatus.Changelist);
+		// Timestamp is used to throttle status requests, so update it to current time:
+		*ChangelistState = CLStatus;
+		ChangelistState->TimeStamp = Now;
+		bUpdated = true;
+
+		// Update files states for files in the changelist
+		bool bUpdateFilesStates = (OutCLFilesStates.Num() == OutChangelistsStates.Num());
+		if (bUpdateFilesStates)
+		{
+			ChangelistState->Files.Reset(OutCLFilesStates[StatusIndex].Num());
+			for (const auto& FileState : OutCLFilesStates[StatusIndex])
+			{
+				TSharedRef<FPlasticSourceControlState, ESPMode::ThreadSafe> CachedFileState = GetProvider().GetStateInternal(FileState.LocalFilename);
+				// Don't override "fileinfo" information and the potential LockedByOther state
+				if (CachedFileState->WorkspaceState != EWorkspaceState::LockedByOther)
+				{
+					CachedFileState->WorkspaceState = FileState.WorkspaceState;
+				}
+				CachedFileState->Changelist = CLStatus.Changelist;
+				ChangelistState->Files.AddUnique(CachedFileState);
+			}
+		}
+	}
+
+	if (bCleanupCache)
+	{
+		TArray<FPlasticSourceControlChangelist> ChangelistsToRemove;
+		GetProvider().GetCachedStateByPredicate([this, &ChangelistsToRemove](const FSourceControlChangelistStateRef& InCLState) {
+			TSharedRef<FPlasticSourceControlChangelistState, ESPMode::ThreadSafe> CLState = StaticCastSharedRef<FPlasticSourceControlChangelistState>(InCLState);
+
+			if (Algo::NoneOf(OutChangelistsStates, [&CLState](const FPlasticSourceControlChangelistState& UpdatedCLState) {
+					return CLState->Changelist == UpdatedCLState.Changelist;
+				}))
+			{
+				ChangelistsToRemove.Add(CLState->Changelist);
+			}
+
+			return false;
+			});
+
+		for (const FPlasticSourceControlChangelist& ChangelistToRemove : ChangelistsToRemove)
+		{
+			GetProvider().RemoveChangelistFromCache(ChangelistToRemove);
+		}
+	}
+
+	return bUpdated;
+}
+
+FPlasticSourceControlChangelist GenerateUniqueChangelistName(FPlasticSourceControlProvider& PlasticSourceControlProvider)
+{
+	FPlasticSourceControlChangelist NewChangelist;
+
+	// Generate a unique number for the name of the new changelist: start from current changeset and increment until a number is available as a new changelist number
+	int32 ChangelistNumber = PlasticSourceControlProvider.GetChangesetNumber();
+	bool bNewNumberOk = false;
+	do
+	{
+		ChangelistNumber++;
+		NewChangelist = FPlasticSourceControlChangelist(FString::FromInt(ChangelistNumber));
+		TSharedRef<FPlasticSourceControlChangelistState, ESPMode::ThreadSafe> ChangelistState = PlasticSourceControlProvider.GetStateInternal(NewChangelist);
+		bNewNumberOk = !ChangelistState->Changelist.IsInitialized();
+	} while (!bNewNumberOk);
+	NewChangelist.SetInitialized();
+
+	return NewChangelist;
+}
+
+FPlasticSourceControlChangelist CreatePendingChangelist(FPlasticSourceControlProvider& PlasticSourceControlProvider, const FString& InDescription, EConcurrency::Type InConcurrency, TArray<FString>& InInfoMessages, TArray<FString>& InErrorMessages)
+{
+	FPlasticSourceControlChangelist NewChangelist = GenerateUniqueChangelistName(PlasticSourceControlProvider);
+
+	bool bCommandSuccessful;
+	TArray<FString> Parameters;
+	if (PlasticSourceControlProvider.GetPlasticScmVersion() < s_NewChangelistFileArgsPlasticScmVersion)
+	{
+		Parameters.Add(TEXT("add"));
+		Parameters.Add(TEXT("\"") + NewChangelist.GetName() + TEXT("\""));
+		Parameters.Add(TEXT("\"") + InDescription + TEXT("\""));
+		Parameters.Add(TEXT("--persistent")); // Create a persistent changelist to stay close to Perforce behavior
+		bCommandSuccessful = PlasticSourceControlUtils::RunCommand(TEXT("changelist"), Parameters, TArray<FString>(), InConcurrency, InInfoMessages, InErrorMessages);
+	}
+	else
+	{
+		Parameters.Add(TEXT("create"));
+		const FScopedTempFile ChangelistNameFile(NewChangelist.GetName());
+		Parameters.Add(FString::Printf(TEXT("--namefile=\"%s\""), *FPaths::ConvertRelativePathToFull(ChangelistNameFile.GetFilename())));
+		const FScopedTempFile ChangelistDescriptionFile(InDescription);
+		Parameters.Add(FString::Printf(TEXT("--descriptionfile=\"%s\""), *FPaths::ConvertRelativePathToFull(ChangelistDescriptionFile.GetFilename())));
+		Parameters.Add(TEXT("--persistent")); // Create a persistent changelist to stay close to Perforce behavior
+		UE_LOG(LogSourceControl, Verbose, TEXT("CreatePendingChangelist(%s):\n\"%s\""), *NewChangelist.GetName(), *InDescription);
+		bCommandSuccessful = PlasticSourceControlUtils::RunCommand(TEXT("changelist"), Parameters, TArray<FString>(), InConcurrency, InInfoMessages, InErrorMessages);
+	}
+	if (!bCommandSuccessful)
+	{
+		NewChangelist.Reset();
+	}
+
+	return NewChangelist;
+}
+
+bool EditChangelistDescription(const FPlasticSourceControlProvider& PlasticSourceControlProvider, const FPlasticSourceControlChangelist& InChangelist, const FString& InDescription, EConcurrency::Type InConcurrency, TArray<FString>& InInfoMessages, TArray<FString>& InErrorMessages)
+{
+	TArray<FString> Parameters;
+	Parameters.Add(TEXT("edit"));
+	if (PlasticSourceControlProvider.GetPlasticScmVersion() < s_NewChangelistFileArgsPlasticScmVersion)
+	{
+		Parameters.Add(TEXT("\"") + InChangelist.GetName() + TEXT("\""));
+		Parameters.Add(TEXT("description"));
+		Parameters.Add(TEXT("\"") + InDescription + TEXT("\""));
+		return PlasticSourceControlUtils::RunCommand(TEXT("changelist"), Parameters, TArray<FString>(), InConcurrency, InInfoMessages, InErrorMessages);
+	}
+	else
+	{
+		const FScopedTempFile ChangelistNameFile(InChangelist.GetName());
+		Parameters.Add(FString::Printf(TEXT("--namefile=\"%s\""), *FPaths::ConvertRelativePathToFull(ChangelistNameFile.GetFilename())));
+		Parameters.Add(TEXT("description"));
+		const FScopedTempFile ChangelistDescriptionFile(InDescription);
+		Parameters.Add(FString::Printf(TEXT("--descriptionfile=\"%s\""), *FPaths::ConvertRelativePathToFull(ChangelistDescriptionFile.GetFilename())));
+		UE_LOG(LogSourceControl, Verbose, TEXT("EditChangelistDescription(%s\n%s)"), *InChangelist.GetName(), *InDescription);
+		return PlasticSourceControlUtils::RunCommand(TEXT("changelist"), Parameters, TArray<FString>(), InConcurrency, InInfoMessages, InErrorMessages);
+	}
+}
+
+bool MoveFilesToChangelist(const FPlasticSourceControlProvider& PlasticSourceControlProvider, const FPlasticSourceControlChangelist& InChangelist, const TArray<FString>& InFiles, const EConcurrency::Type InConcurrency, TArray<FString>& OutResults, TArray<FString>& OutErrorMessages)
+{
+	if (InFiles.Num() > 0)
+	{
+		TArray<FString> Parameters;
+		if (PlasticSourceControlProvider.GetPlasticScmVersion() < s_NewChangelistFileArgsPlasticScmVersion)
+		{
+			Parameters.Add(TEXT("\"") + InChangelist.GetName() + TEXT("\""));
+			Parameters.Add(TEXT("add"));
+			return PlasticSourceControlUtils::RunCommand(TEXT("changelist"), Parameters, InFiles, InConcurrency, OutResults, OutErrorMessages);
+		}
+		else
+		{
+			const FScopedTempFile ChangelistNameFile(InChangelist.GetName());
+			Parameters.Add(FString::Printf(TEXT("--namefile=\"%s\""), *FPaths::ConvertRelativePathToFull(ChangelistNameFile.GetFilename())));
+			Parameters.Add(TEXT("add"));
+			UE_LOG(LogSourceControl, Verbose, TEXT("MoveFilesToChangelist(%s)"), *InChangelist.GetName());
+			return PlasticSourceControlUtils::RunCommand(TEXT("changelist"), Parameters, InFiles, InConcurrency, OutResults, OutErrorMessages);
+		}
+	}
+	return true;
+}
+
+FPlasticNewChangelistWorker::FPlasticNewChangelistWorker(FPlasticSourceControlProvider& InSourceControlProvider)
+	: IPlasticSourceControlWorker(InSourceControlProvider)
+	, NewChangelistState(NewChangelist)
+{
+}
+
+FName FPlasticNewChangelistWorker::GetName() const
+{
+	return "NewChangelist";
+}
+
+bool FPlasticNewChangelistWorker::Execute(class FPlasticSourceControlCommand& InCommand)
+{
+	TRACE_CPUPROFILER_EVENT_SCOPE(FPlasticNewChangelistWorker::Execute);
+
+	check(InCommand.Operation->GetName() == GetName());
+	TSharedRef<FNewChangelist, ESPMode::ThreadSafe> Operation = StaticCastSharedRef<FNewChangelist>(InCommand.Operation);
+
+	FString Description = Operation->GetDescription().ToString();
+	// Note: old "cm" doesn't support newlines, quotes, and question marks on changelist's name or description
+	if (GetProvider().GetPlasticScmVersion() < s_NewChangelistFileArgsPlasticScmVersion)
+	{
+		Description.ReplaceInline(TEXT("\r\n"), TEXT(" "), ESearchCase::CaseSensitive);
+		Description.ReplaceCharInline(TEXT('\n'), TEXT(' '), ESearchCase::CaseSensitive);
+		Description.ReplaceCharInline(TEXT('\"'), TEXT('\''), ESearchCase::CaseSensitive);
+		Description.ReplaceCharInline(TEXT('?'), TEXT('.'), ESearchCase::CaseSensitive);
+		Description.ReplaceCharInline(TEXT('*'), TEXT('.'), ESearchCase::CaseSensitive);
+	}
+
+	// Create a new numbered persistent changelist ala Perforce
+	NewChangelist = CreatePendingChangelist(GetProvider(), Description, InCommand.Concurrency, InCommand.InfoMessages, InCommand.ErrorMessages);
+
+	// Successfully created new changelist
+	if (NewChangelist.IsInitialized())
+	{
+		InCommand.bCommandSuccessful = true;
+
+		NewChangelistState.Changelist = NewChangelist;
+		NewChangelistState.Description = MoveTemp(Description);
+
+		Operation->SetNewChangelist(MakeShared<FPlasticSourceControlChangelist>(NewChangelist));
+
+		if (InCommand.Files.Num() > 0)
+		{
+			InCommand.bCommandSuccessful = MoveFilesToChangelist(GetProvider(), NewChangelist, InCommand.Files, InCommand.Concurrency, InCommand.InfoMessages, InCommand.ErrorMessages);
+			if (InCommand.bCommandSuccessful)
+			{
+				MovedFiles = InCommand.Files;
+			}
+		}
+	}
+
+	return InCommand.bCommandSuccessful;
+}
+
+bool FPlasticNewChangelistWorker::UpdateStates()
+{
+	if (NewChangelist.IsInitialized())
+	{
+		const FDateTime Now = FDateTime::Now();
+
+		TSharedRef<FPlasticSourceControlChangelistState, ESPMode::ThreadSafe> ChangelistState = GetProvider().GetStateInternal(NewChangelist);
+		*ChangelistState = NewChangelistState;
+		ChangelistState->TimeStamp = Now;
+
+		// 3 things to do here:
+		for (const FString& MovedFile : MovedFiles)
+		{
+			TSharedRef<FPlasticSourceControlState, ESPMode::ThreadSafe> FileState = GetProvider().GetStateInternal(MovedFile);
+
+			// 1- Remove these files from their previous changelist
+			TSharedRef<FPlasticSourceControlChangelistState, ESPMode::ThreadSafe> PreviousChangelist = GetProvider().GetStateInternal(FileState->Changelist);
+			PreviousChangelist->Files.Remove(FileState);
+
+			// 2- Add to the new changelist
+			ChangelistState->Files.Add(FileState);
+
+			// 3- Update changelist in file state
+			FileState->Changelist = NewChangelist;
+			FileState->TimeStamp = Now;
+		}
+
+		return true;
+	}
+	else
+	{
+		return false;
+	}
+}
+
+
+FName FPlasticDeleteChangelistWorker::GetName() const
+{
+	return "DeleteChangelist";
+}
+
+bool FPlasticDeleteChangelistWorker::Execute(class FPlasticSourceControlCommand& InCommand)
+{
+	TRACE_CPUPROFILER_EVENT_SCOPE(FPlasticDeleteChangelistWorker::Execute);
+
+	// Can't delete the default changelist
+	if (InCommand.Changelist.IsDefault())
+	{
+		InCommand.bCommandSuccessful = false;
+	}
+	else
+	{
+		check(InCommand.Operation->GetName() == GetName());
+		TSharedRef<FDeleteChangelist, ESPMode::ThreadSafe> Operation = StaticCastSharedRef<FDeleteChangelist>(InCommand.Operation);
+
+		InCommand.bCommandSuccessful = DeleteChangelist(GetProvider(), InCommand.Changelist, InCommand.Concurrency, InCommand.InfoMessages, InCommand.ErrorMessages);
+
+		// NOTE: for now it is not possible to delete a changelist with files through the Editor
+		if (InCommand.Files.Num() > 0 && InCommand.bCommandSuccessful)
+		{
+			TSharedRef<FPlasticSourceControlChangelistState, ESPMode::ThreadSafe> ChangelistState = GetProvider().GetStateInternal(InCommand.Changelist);
+			const TArray<FString> Files = FileNamesFromFileStates(ChangelistState->Files);
+			InCommand.bCommandSuccessful = MoveFilesToChangelist(GetProvider(), FPlasticSourceControlChangelist::DefaultChangelist, Files, InCommand.Concurrency, InCommand.InfoMessages, InCommand.ErrorMessages);
+		}
+
+		// Keep track of changelist to update the cache
+		if (InCommand.bCommandSuccessful)
+		{
+			DeletedChangelist = InCommand.Changelist;
+		}
+	}
+
+	return InCommand.bCommandSuccessful;
+}
+
+bool FPlasticDeleteChangelistWorker::UpdateStates()
+{
+	TRACE_CPUPROFILER_EVENT_SCOPE(FPlasticDeleteChangelistWorker::UpdateStates);
+
+	if (DeletedChangelist.IsInitialized())
+	{
+		return GetProvider().RemoveChangelistFromCache(DeletedChangelist);
+	}
+
+	return false;
+}
+
+
+FName FPlasticEditChangelistWorker::GetName() const
+{
+	return "EditChangelist";
+}
+
+bool FPlasticEditChangelistWorker::Execute(class FPlasticSourceControlCommand& InCommand)
+{
+	TRACE_CPUPROFILER_EVENT_SCOPE(FPlasticEditChangelistWorker::Execute);
+
+	check(InCommand.Operation->GetName() == GetName());
+	TSharedRef<FEditChangelist, ESPMode::ThreadSafe> Operation = StaticCastSharedRef<FEditChangelist>(InCommand.Operation);
+
+	EditedDescription = Operation->GetDescription().ToString();
+	// Note: old "cm" doesn't support newlines, quotes, and question marks on changelist's name or description
+	if (GetProvider().GetPlasticScmVersion() < s_NewChangelistFileArgsPlasticScmVersion)
+	{
+		EditedDescription.ReplaceInline(TEXT("\r\n"), TEXT(" "), ESearchCase::CaseSensitive);
+		EditedDescription.ReplaceCharInline(TEXT('\n'), TEXT(' '), ESearchCase::CaseSensitive);
+		EditedDescription.ReplaceCharInline(TEXT('\"'), TEXT('\''), ESearchCase::CaseSensitive);
+		EditedDescription.ReplaceCharInline(TEXT('?'), TEXT('.'), ESearchCase::CaseSensitive);
+		EditedDescription.ReplaceCharInline(TEXT('*'), TEXT('.'), ESearchCase::CaseSensitive);
+	}
+
+	if (InCommand.Changelist.IsDefault())
+	{
+		// Create a new numbered persistent changelist since we cannot edit the default changelist
+		EditedChangelist = CreatePendingChangelist(GetProvider(), EditedDescription, InCommand.Concurrency, InCommand.InfoMessages, InCommand.ErrorMessages);
+		if (EditedChangelist.IsInitialized())
+		{
+			// And then move all its files to the new changelist
+			TSharedRef<FPlasticSourceControlChangelistState, ESPMode::ThreadSafe> ChangelistState = GetProvider().GetStateInternal(InCommand.Changelist);
+			ReopenedFiles = FileNamesFromFileStates(ChangelistState->Files);
+			InCommand.bCommandSuccessful = MoveFilesToChangelist(GetProvider(), EditedChangelist, ReopenedFiles, InCommand.Concurrency, InCommand.InfoMessages, InCommand.ErrorMessages);
+		}
+	}
+	else
+	{
+		InCommand.bCommandSuccessful = EditChangelistDescription(GetProvider(), InCommand.Changelist, EditedDescription, InCommand.Concurrency, InCommand.InfoMessages, InCommand.ErrorMessages);
+		if (InCommand.bCommandSuccessful)
+		{
+			EditedChangelist = InCommand.Changelist;
+		}
+	}
+
+	return InCommand.bCommandSuccessful;
+}
+
+bool FPlasticEditChangelistWorker::UpdateStates()
+{
+	TRACE_CPUPROFILER_EVENT_SCOPE(FPlasticEditChangelistWorker::UpdateStates);
+
+	if (EditedChangelist.IsInitialized())
+	{
+		const FDateTime Now = FDateTime::Now();
+		TSharedRef<FPlasticSourceControlChangelistState, ESPMode::ThreadSafe> EditedChangelistState = GetProvider().GetStateInternal(EditedChangelist);
+		EditedChangelistState->Description = EditedDescription;
+		EditedChangelistState->Changelist = EditedChangelist;
+		EditedChangelistState->TimeStamp = Now;
+
+		// 3 things to do here:
+		for (const FString& ReopenedFile : ReopenedFiles)
+		{
+			TSharedRef<FPlasticSourceControlState, ESPMode::ThreadSafe> FileState = GetProvider().GetStateInternal(ReopenedFile);
+
+			// 1- Remove these files from their previous changelist
+			TSharedRef<FPlasticSourceControlChangelistState, ESPMode::ThreadSafe> PreviousChangelist = GetProvider().GetStateInternal(FileState->Changelist);
+			PreviousChangelist->Files.Remove(FileState);
+
+			// 2- Add to the new changelist
+			EditedChangelistState->Files.Add(FileState);
+
+			// 3- Update changelist in file state
+			FileState->Changelist = EditedChangelist;
+			FileState->TimeStamp = Now;
+		}
+
+		return true;
+	}
+	else
+	{
+		return false;
+	}
+}
+
+
+FName FPlasticReopenWorker::GetName() const
+{
+	return "MoveToChangelist";
+}
+
+bool FPlasticReopenWorker::Execute(FPlasticSourceControlCommand& InCommand)
+{
+	TRACE_CPUPROFILER_EVENT_SCOPE(FPlasticReopenWorker::Execute);
+
+	check(InCommand.Operation->GetName() == GetName());
+
+	InCommand.bCommandSuccessful = MoveFilesToChangelist(GetProvider(), InCommand.Changelist, InCommand.Files, InCommand.Concurrency, InCommand.InfoMessages, InCommand.ErrorMessages);
+	if (InCommand.bCommandSuccessful)
+	{
+		ReopenedFiles = InCommand.Files;
+		DestinationChangelist = InCommand.Changelist;
+	}
+
+	return InCommand.bCommandSuccessful;
+}
+
+bool FPlasticReopenWorker::UpdateStates()
+{
+	TRACE_CPUPROFILER_EVENT_SCOPE(FPlasticReopenWorker::UpdateStates);
+
+	if (DestinationChangelist.IsInitialized())
+	{
+		const FDateTime Now = FDateTime::Now();
+		TSharedRef<FPlasticSourceControlChangelistState, ESPMode::ThreadSafe> DestinationChangelistState = GetProvider().GetStateInternal(DestinationChangelist);
+
+		// 3 things to do here:
+		for (const FString& ReopenedFile : ReopenedFiles)
+		{
+			TSharedRef<FPlasticSourceControlState, ESPMode::ThreadSafe> FileState = GetProvider().GetStateInternal(ReopenedFile);
+
+			// 1- Remove these files from their previous changelist
+			TSharedRef<FPlasticSourceControlChangelistState, ESPMode::ThreadSafe> PreviousChangelist = GetProvider().GetStateInternal(FileState->Changelist);
+			PreviousChangelist->Files.Remove(FileState);
+
+			// 2- Add to the new changelist
+			DestinationChangelistState->Files.Add(FileState);
+
+			// 3- Update changelist in file state
+			FileState->Changelist = DestinationChangelist;
+			FileState->TimeStamp = Now;
+		}
+
+		return ReopenedFiles.Num() > 0;
+	}
+	else
+	{
+		return false;
+	}
+}
 #undef LOCTEXT_NAMESPACE
diff --git a/Engine/Plugins/Developer/PlasticSourceControl/Source/PlasticSourceControl/Private/PlasticSourceControlOperations.h b/Engine/Plugins/Developer/PlasticSourceControl/Source/PlasticSourceControl/Private/PlasticSourceControlOperations.h
--- a/Engine/Plugins/Developer/PlasticSourceControl/Source/PlasticSourceControl/Private/PlasticSourceControlOperations.h
+++ b/Engine/Plugins/Developer/PlasticSourceControl/Source/PlasticSourceControl/Private/PlasticSourceControlOperations.h
@@ -6,16 +6,19 @@
 
 #include "CoreMinimal.h"
 #include "IPlasticSourceControlWorker.h"
-#include "PlasticSourceControlState.h"
 #include "PlasticSourceControlRevision.h"
-
+#include "PlasticSourceControlState.h"
 #include "ISourceControlOperation.h"
 
+#include "PlasticSourceControlChangelist.h"
+#include "PlasticSourceControlChangelistState.h"
+
+class FPlasticSourceControlProvider;
 
 /**
  * Internal operation used to revert checked-out unchanged files
 */
-class FPlasticRevertUnchanged : public ISourceControlOperation
+class FPlasticRevertUnchanged final : public ISourceControlOperation
 {
 public:
 	// ISourceControlOperation interface
@@ -28,7 +31,7 @@ public:
 /**
  * Internal operation used to revert checked-out files
 */
-class FPlasticRevertAll : public ISourceControlOperation
+class FPlasticRevertAll final : public ISourceControlOperation
 {
 public:
 	// ISourceControlOperation interface
@@ -41,7 +44,7 @@ public:
 /**
 * Internal operation used to initialize a new Workspace and a new Repository
 */
-class FPlasticMakeWorkspace : public ISourceControlOperation
+class FPlasticMakeWorkspace final : public ISourceControlOperation
 {
 public:
 	// ISourceControlOperation interface
@@ -57,10 +60,13 @@ public:
 
 /** Called when first activated on a project, and then at project load time.
  *  Look for the root directory of the Plastic workspace (where the ".plastic/" subdirectory is located). */
-class FPlasticConnectWorker : public IPlasticSourceControlWorker
+class FPlasticConnectWorker final : public IPlasticSourceControlWorker
 {
 public:
-	virtual ~FPlasticConnectWorker() {}
+	explicit FPlasticConnectWorker(FPlasticSourceControlProvider& InSourceControlProvider)
+		: IPlasticSourceControlWorker(InSourceControlProvider)
+	{}
+	virtual ~FPlasticConnectWorker() = default;
 	// IPlasticSourceControlWorker interface
 	virtual FName GetName() const override;
 	virtual bool Execute(class FPlasticSourceControlCommand& InCommand) override;
@@ -71,10 +77,14 @@ public:
 	TArray<FPlasticSourceControlState> States;
 };
 
-class FPlasticCheckOutWorker : public IPlasticSourceControlWorker
+class FPlasticCheckOutWorker final : public IPlasticSourceControlWorker
 {
 public:
-	virtual ~FPlasticCheckOutWorker() {}
+	explicit FPlasticCheckOutWorker(FPlasticSourceControlProvider& InSourceControlProvider)
+		: IPlasticSourceControlWorker(InSourceControlProvider)
+		, InChangelist(FPlasticSourceControlChangelist::DefaultChangelist) // By default, add checked out files in the default changelist.
+	{}
+	virtual ~FPlasticCheckOutWorker() = default;
 	// IPlasticSourceControlWorker interface
 	virtual FName GetName() const override;
 	virtual bool Execute(class FPlasticSourceControlCommand& InCommand) override;
@@ -83,13 +93,19 @@ public:
 public:
 	/** Temporary states for results */
 	TArray<FPlasticSourceControlState> States;
+
+	/** Changelist we checked-out files to (defaults to the Default changelist) */
+	FPlasticSourceControlChangelist InChangelist;
 };
 
 /** Check-in a set of file to the local depot. */
-class FPlasticCheckInWorker : public IPlasticSourceControlWorker
+class FPlasticCheckInWorker final : public IPlasticSourceControlWorker
 {
 public:
-	virtual ~FPlasticCheckInWorker() {}
+	explicit FPlasticCheckInWorker(FPlasticSourceControlProvider& InSourceControlProvider)
+		: IPlasticSourceControlWorker(InSourceControlProvider)
+	{}
+	virtual ~FPlasticCheckInWorker() = default;
 	// IPlasticSourceControlWorker interface
 	virtual FName GetName() const override;
 	virtual bool Execute(class FPlasticSourceControlCommand& InCommand) override;
@@ -98,13 +114,20 @@ public:
 public:
 	/** Temporary states for results */
 	TArray<FPlasticSourceControlState> States;
+
+	/** Changelist we submitted */
+	FPlasticSourceControlChangelist InChangelist;
 };
 
 /** Add an untracked file to source control (so only a subset of the Plastic add command). */
-class FPlasticMarkForAddWorker : public IPlasticSourceControlWorker
+class FPlasticMarkForAddWorker final : public IPlasticSourceControlWorker
 {
 public:
-	virtual ~FPlasticMarkForAddWorker() {}
+	explicit FPlasticMarkForAddWorker(FPlasticSourceControlProvider& InSourceControlProvider)
+		: IPlasticSourceControlWorker(InSourceControlProvider)
+		, InChangelist(FPlasticSourceControlChangelist::DefaultChangelist) // By default, add new files in the default changelist.
+	{}
+	virtual ~FPlasticMarkForAddWorker() = default;
 	// IPlasticSourceControlWorker interface
 	virtual FName GetName() const override;
 	virtual bool Execute(class FPlasticSourceControlCommand& InCommand) override;
@@ -113,13 +136,20 @@ public:
 public:
 	/** Temporary states for results */
 	TArray<FPlasticSourceControlState> States;
+
+	/** Changelist we added files to (defaults to the Default changelist) */
+	FPlasticSourceControlChangelist InChangelist;
 };
 
 /** Delete a file and remove it from source control. */
-class FPlasticDeleteWorker : public IPlasticSourceControlWorker
+class FPlasticDeleteWorker final : public IPlasticSourceControlWorker
 {
 public:
-	virtual ~FPlasticDeleteWorker() {}
+	explicit FPlasticDeleteWorker(FPlasticSourceControlProvider& InSourceControlProvider)
+		: IPlasticSourceControlWorker(InSourceControlProvider)
+		, InChangelist(FPlasticSourceControlChangelist::DefaultChangelist) // By default, add deleted files in the default changelist.
+	{}
+	virtual ~FPlasticDeleteWorker() = default;
 	// IPlasticSourceControlWorker interface
 	virtual FName GetName() const override;
 	virtual bool Execute(class FPlasticSourceControlCommand& InCommand) override;
@@ -128,24 +158,37 @@ public:
 public:
 	/** Temporary states for results */
 	TArray<FPlasticSourceControlState> States;
+
+	/** Changelist we delete files to (defaults to the Default changelist) */
+	FPlasticSourceControlChangelist InChangelist;
 };
 
 /** Revert any change to a file to its state on the local depot. */
-class FPlasticRevertWorker : public IPlasticSourceControlWorker
+class FPlasticRevertWorker final : public IPlasticSourceControlWorker
 {
 public:
-	virtual ~FPlasticRevertWorker() {}
+	explicit FPlasticRevertWorker(FPlasticSourceControlProvider& InSourceControlProvider)
+		: IPlasticSourceControlWorker(InSourceControlProvider)
+	{}
+	virtual ~FPlasticRevertWorker() = default;
 	// IPlasticSourceControlWorker interface
 	virtual FName GetName() const override;
 	virtual bool Execute(class FPlasticSourceControlCommand& InCommand) override;
 	virtual bool UpdateStates() override;
+
+public:
+	/** Temporary states for results */
+	TArray<FPlasticSourceControlState> States;
 };
 
 /** Revert only unchanged file(s) (uncheckout). */
-class FPlasticRevertUnchangedWorker : public IPlasticSourceControlWorker
+class FPlasticRevertUnchangedWorker final : public IPlasticSourceControlWorker
 {
 public:
-	virtual ~FPlasticRevertUnchangedWorker() {}
+	explicit FPlasticRevertUnchangedWorker(FPlasticSourceControlProvider& InSourceControlProvider)
+		: IPlasticSourceControlWorker(InSourceControlProvider)
+	{}
+	virtual ~FPlasticRevertUnchangedWorker() = default;
 	// IPlasticSourceControlWorker interface
 	virtual FName GetName() const override;
 	virtual bool Execute(class FPlasticSourceControlCommand& InCommand) override;
@@ -157,10 +200,13 @@ public:
 };
 
 /** Revert all checked-out file(s). */
-class FPlasticRevertAllWorker : public IPlasticSourceControlWorker
+class FPlasticRevertAllWorker final : public IPlasticSourceControlWorker
 {
 public:
-	virtual ~FPlasticRevertAllWorker() {}
+	explicit FPlasticRevertAllWorker(FPlasticSourceControlProvider& InSourceControlProvider)
+		: IPlasticSourceControlWorker(InSourceControlProvider)
+	{}
+	virtual ~FPlasticRevertAllWorker() = default;
 	// IPlasticSourceControlWorker interface
 	virtual FName GetName() const override;
 	virtual bool Execute(class FPlasticSourceControlCommand& InCommand) override;
@@ -172,10 +218,13 @@ public:
 };
 
 /** Initialize a new Workspace and a new Repository */
-class FPlasticMakeWorkspaceWorker : public IPlasticSourceControlWorker
+class FPlasticMakeWorkspaceWorker final : public IPlasticSourceControlWorker
 {
 public:
-	virtual ~FPlasticMakeWorkspaceWorker() {}
+	explicit FPlasticMakeWorkspaceWorker(FPlasticSourceControlProvider& InSourceControlProvider)
+		: IPlasticSourceControlWorker(InSourceControlProvider)
+	{}
+	virtual ~FPlasticMakeWorkspaceWorker() = default;
 	// IPlasticSourceControlWorker interface
 	virtual FName GetName() const override;
 	virtual bool Execute(class FPlasticSourceControlCommand& InCommand) override;
@@ -183,10 +232,13 @@ public:
 };
 
 /** Plastic update the workspace to latest changes */
-class FPlasticSyncWorker : public IPlasticSourceControlWorker
+class FPlasticSyncWorker final : public IPlasticSourceControlWorker
 {
 public:
-	virtual ~FPlasticSyncWorker() {}
+	explicit FPlasticSyncWorker(FPlasticSourceControlProvider& InSourceControlProvider)
+		: IPlasticSourceControlWorker(InSourceControlProvider)
+	{}
+	virtual ~FPlasticSyncWorker() = default;
 	// IPlasticSourceControlWorker interface
 	virtual FName GetName() const override;
 	virtual bool Execute(class FPlasticSourceControlCommand& InCommand) override;
@@ -198,10 +250,13 @@ public:
 };
 
 /** Get source control status of files on local workspace. */
-class FPlasticUpdateStatusWorker : public IPlasticSourceControlWorker
+class FPlasticUpdateStatusWorker final : public IPlasticSourceControlWorker
 {
 public:
-	virtual ~FPlasticUpdateStatusWorker() {}
+	explicit FPlasticUpdateStatusWorker(FPlasticSourceControlProvider& InSourceControlProvider)
+		: IPlasticSourceControlWorker(InSourceControlProvider)
+	{}
+	virtual ~FPlasticUpdateStatusWorker() = default;
 	// IPlasticSourceControlWorker interface
 	virtual FName GetName() const override;
 	virtual bool Execute(class FPlasticSourceControlCommand& InCommand) override;
@@ -213,10 +268,13 @@ public:
 };
 
 /** Copy or Move operation on a single file */
-class FPlasticCopyWorker : public IPlasticSourceControlWorker
+class FPlasticCopyWorker final : public IPlasticSourceControlWorker
 {
 public:
-	virtual ~FPlasticCopyWorker() {}
+	explicit FPlasticCopyWorker(FPlasticSourceControlProvider& InSourceControlProvider)
+		: IPlasticSourceControlWorker(InSourceControlProvider)
+	{}
+	virtual ~FPlasticCopyWorker() = default;
 	// IPlasticSourceControlWorker interface
 	virtual FName GetName() const override;
 	virtual bool Execute(class FPlasticSourceControlCommand& InCommand) override;
@@ -228,10 +286,13 @@ public:
 };
 
 /** Plastic command to mark the conflict as solved */
-class FPlasticResolveWorker : public IPlasticSourceControlWorker
+class FPlasticResolveWorker final : public IPlasticSourceControlWorker
 {
 public:
-	virtual ~FPlasticResolveWorker() {}
+	explicit FPlasticResolveWorker(FPlasticSourceControlProvider& InSourceControlProvider)
+		: IPlasticSourceControlWorker(InSourceControlProvider)
+	{}
+	virtual ~FPlasticResolveWorker() = default;
 	virtual FName GetName() const override;
 	virtual bool Execute(class FPlasticSourceControlCommand& InCommand) override;
 	virtual bool UpdateStates() override;
@@ -240,3 +301,101 @@ private:
 	/** Temporary states for results */
 	TArray<FPlasticSourceControlState> States;
 };
+
+class FPlasticGetPendingChangelistsWorker final : public IPlasticSourceControlWorker
+{
+public:
+	explicit FPlasticGetPendingChangelistsWorker(FPlasticSourceControlProvider& InSourceControlProvider)
+		: IPlasticSourceControlWorker(InSourceControlProvider)
+	{}
+	virtual ~FPlasticGetPendingChangelistsWorker() = default;
+	// IPlasticSourceControlWorker interface
+	virtual FName GetName() const override;
+	virtual bool Execute(class FPlasticSourceControlCommand& InCommand) override;
+	virtual bool UpdateStates() override;
+
+public:
+	/** Temporary states for results */
+	TArray<FPlasticSourceControlChangelistState> OutChangelistsStates;
+	TArray<TArray<FPlasticSourceControlState>> OutCLFilesStates;
+
+private:
+	/** Controls whether or not we will remove changelists from the cache after a full update */
+	bool bCleanupCache = false;
+};
+
+class FPlasticNewChangelistWorker final : public IPlasticSourceControlWorker
+{
+public:
+	explicit FPlasticNewChangelistWorker(FPlasticSourceControlProvider& InSourceControlProvider);
+	virtual ~FPlasticNewChangelistWorker() = default;
+	// IPlasticSourceControlWorker interface
+	virtual FName GetName() const override;
+	virtual bool Execute(class FPlasticSourceControlCommand& InCommand) override;
+	virtual bool UpdateStates() override;
+
+public:
+	/** New changelist information */
+	FPlasticSourceControlChangelist NewChangelist;
+	FPlasticSourceControlChangelistState NewChangelistState;
+
+	/** Files that were moved */
+	TArray<FString> MovedFiles;
+};
+
+class FPlasticDeleteChangelistWorker final : public IPlasticSourceControlWorker
+{
+public:
+	explicit FPlasticDeleteChangelistWorker(FPlasticSourceControlProvider& InSourceControlProvider)
+		: IPlasticSourceControlWorker(InSourceControlProvider)
+	{}
+	virtual ~FPlasticDeleteChangelistWorker() = default;
+	// IPlasticSourceControlWorker interface
+	virtual FName GetName() const override;
+	virtual bool Execute(class FPlasticSourceControlCommand& InCommand) override;
+	virtual bool UpdateStates() override;
+
+public:
+	FPlasticSourceControlChangelist DeletedChangelist;
+};
+
+class FPlasticEditChangelistWorker final : public IPlasticSourceControlWorker
+{
+public:
+	explicit FPlasticEditChangelistWorker(FPlasticSourceControlProvider& InSourceControlProvider)
+		: IPlasticSourceControlWorker(InSourceControlProvider)
+	{}
+	virtual ~FPlasticEditChangelistWorker() = default;
+	// IPlasticSourceControlWorker interface
+	virtual FName GetName() const override;
+	virtual bool Execute(class FPlasticSourceControlCommand& InCommand) override;
+	virtual bool UpdateStates() override;
+
+public:
+	FPlasticSourceControlChangelist EditedChangelist;
+	FString EditedDescription;
+
+	/** Reopened files (moved to a new changelist, if any, when editing the Default changelist) */
+	TArray<FString> ReopenedFiles;
+};
+
+class FPlasticReopenWorker final : public IPlasticSourceControlWorker
+{
+public:
+	explicit FPlasticReopenWorker(FPlasticSourceControlProvider& InSourceControlProvider)
+		: IPlasticSourceControlWorker(InSourceControlProvider)
+	{}
+	virtual ~FPlasticReopenWorker() = default;
+	// IPlasticSourceControlWorker interface
+	virtual FName GetName() const override;
+	virtual bool Execute(class FPlasticSourceControlCommand& InCommand) override;
+	virtual bool UpdateStates() override;
+
+protected:
+	/** Reopened files (moved to a new changelist) */
+	TArray<FString> ReopenedFiles;
+
+	/** Destination changelist */
+	FPlasticSourceControlChangelist DestinationChangelist;
+};
+
diff --git a/Engine/Plugins/Developer/PlasticSourceControl/Source/PlasticSourceControl/Private/PlasticSourceControlProjectSettings.h b/Engine/Plugins/Developer/PlasticSourceControl/Source/PlasticSourceControl/Private/PlasticSourceControlProjectSettings.h
new file mode 100644
--- /dev/null
+++ b/Engine/Plugins/Developer/PlasticSourceControl/Source/PlasticSourceControl/Private/PlasticSourceControlProjectSettings.h
@@ -0,0 +1,27 @@
+// Copyright Epic Games, Inc. All Rights Reserved.
+
+#pragma once
+
+#include "CoreMinimal.h"
+#include "Engine/DeveloperSettings.h"
+#include "PlasticSourceControlProjectSettings.generated.h"
+
+/** Project Settings for Plastic SCM Source Control. Saved in Config/DefaultEditor.ini */
+UCLASS(config = Editor, defaultconfig, meta = (DisplayName = "Source Control - Plastic SCM"))
+class UPlasticSourceControlProjectSettings : public UDeveloperSettings
+{
+	GENERATED_BODY()
+
+public:
+	/** Map Plastic SCM user names (typically e-mail addresses or company domain names) to display names for brevity. */
+	UPROPERTY(config, EditAnywhere, Category = "Plastic SCM")
+	TMap<FString, FString> UserNameToDisplayName;
+
+	/** Hide the domain part of an username e-mail address (eg @gmail.com) if the UserNameToDisplayName map didn't match (enabled by default). */
+	UPROPERTY(config, EditAnywhere, Category = "Plastic SCM")
+	bool bHideEmailDomainInUsername = true;
+
+	/** If enabled, you'll be prompted to check out changed files (enabled by default). Checkout is needed to work with Changelists. */
+	UPROPERTY(config, EditAnywhere, Category = "Plastic SCM")
+	bool bPromptForCheckoutOnChange = true;
+};
diff --git a/Engine/Plugins/Developer/PlasticSourceControl/Source/PlasticSourceControl/Private/PlasticSourceControlProvider.cpp b/Engine/Plugins/Developer/PlasticSourceControl/Source/PlasticSourceControl/Private/PlasticSourceControlProvider.cpp
--- a/Engine/Plugins/Developer/PlasticSourceControl/Source/PlasticSourceControl/Private/PlasticSourceControlProvider.cpp
+++ b/Engine/Plugins/Developer/PlasticSourceControl/Source/PlasticSourceControl/Private/PlasticSourceControlProvider.cpp
@@ -2,9 +2,10 @@
 
 #include "PlasticSourceControlProvider.h"
 
+#include "PlasticSourceControlChangelistState.h"
 #include "PlasticSourceControlCommand.h"
-#include "PlasticSourceControlModule.h"
 #include "PlasticSourceControlOperations.h"
+#include "PlasticSourceControlProjectSettings.h"
 #include "PlasticSourceControlSettings.h"
 #include "PlasticSourceControlState.h"
 #include "PlasticSourceControlUtils.h"
@@ -17,7 +18,9 @@
 #include "SourceControlOperations.h"
 #include "Interfaces/IPluginManager.h"
 
+#include "Algo/Transform.h"
 #include "Misc/Paths.h"
+#include "Misc/MessageDialog.h"
 #include "HAL/PlatformProcess.h"
 #include "Misc/QueuedThreadPool.h"
 
@@ -25,6 +28,11 @@
 
 static FName ProviderName("Plastic SCM");
 
+FPlasticSourceControlProvider::FPlasticSourceControlProvider()
+{
+	PlasticSourceControlSettings.LoadSettings();
+}
+
 void FPlasticSourceControlProvider::Init(bool bForceConnection)
 {
 	// Init() is called multiple times at startup: do not check Plastic SCM each time
@@ -40,14 +48,13 @@ void FPlasticSourceControlProvider::Init(bool bForceConnection)
 		CheckPlasticAvailability();
 
 		// Override the source control logs verbosity level if needed based on settings
-		FPlasticSourceControlModule& PlasticSourceControl = FModuleManager::LoadModuleChecked<FPlasticSourceControlModule>("PlasticSourceControl");
-		if (PlasticSourceControl.AccessSettings().GetEnableVerboseLogs())
+		if (AccessSettings().GetEnableVerboseLogs())
 		{
 			PlasticSourceControlUtils::SwitchVerboseLogs(true);
 		}
 	}
 
-	if (bForceConnection && !bServerAvailable)
+	if (bForceConnection && bPlasticAvailable && bWorkspaceFound && !bServerAvailable)
 	{
 		// Execute a 'checkconnection' command to set bServerAvailable based on the connectivity of the server
 		TArray<FString> InfoMessages, ErrorMessages;
@@ -66,8 +73,7 @@ void FPlasticSourceControlProvider::Init(bool bForceConnection)
 
 void FPlasticSourceControlProvider::CheckPlasticAvailability()
 {
-	FPlasticSourceControlModule& PlasticSourceControl = FModuleManager::LoadModuleChecked<FPlasticSourceControlModule>("PlasticSourceControl");
-	FString PathToPlasticBinary = PlasticSourceControl.AccessSettings().GetBinaryPath();
+	FString PathToPlasticBinary = AccessSettings().GetBinaryPath();
 	if (PathToPlasticBinary.IsEmpty())
 	{
 		bPlasticAvailable = false;
@@ -76,7 +82,7 @@ void FPlasticSourceControlProvider::CheckPlasticAvailability()
 		PathToPlasticBinary = PlasticSourceControlUtils::FindPlasticBinaryPath();
 		if (!PathToPlasticBinary.IsEmpty())
 		{
-			PlasticSourceControl.AccessSettings().SetBinaryPath(PathToPlasticBinary);
+			AccessSettings().SetBinaryPath(PathToPlasticBinary);
 		}
 	}
 
@@ -88,20 +94,29 @@ void FPlasticSourceControlProvider::CheckPlasticAvailability()
 
 		// Launch the Plastic SCM cli shell on the background to issue all commands during this session
 		bPlasticAvailable = PlasticSourceControlUtils::LaunchBackgroundPlasticShell(PathToPlasticBinary, PathToWorkspaceRoot);
-		if (bPlasticAvailable)
+		if (!bPlasticAvailable)
 		{
-			PlasticSourceControlUtils::GetPlasticScmVersion(PlasticScmVersion);
+			return;
+		}
 
-			// Get user name (from the global Plastic SCM client config)
-			PlasticSourceControlUtils::GetUserName(UserName);
+		bPlasticAvailable = PlasticSourceControlUtils::GetPlasticScmVersion(PlasticScmVersion);
+		if (!bPlasticAvailable)
+		{
+			return;
+		}
 
-			// Register Console Commands
-			PlasticSourceControlConsole.Register();
+		// Get user name (from the global Plastic SCM client config)
+		PlasticSourceControlUtils::GetUserName(UserName);
 
-			if (!bWorkspaceFound)
-			{
-				UE_LOG(LogSourceControl, Warning, TEXT("'%s' is not part of a Plastic workspace"), *FPaths::ProjectDir());
-			}
+		// Register Console Commands
+		PlasticSourceControlConsole.Register();
+
+		if (!bWorkspaceFound)
+		{
+			// This info message is only useful here, if bPlasticAvailable, for the Login window
+			FFormatNamedArguments Args;
+			Args.Add(TEXT("WorkspacePath"), FText::FromString(PathToWorkspaceRoot));
+			FMessageLog("SourceControl").Info(FText::Format(LOCTEXT("NotInAWorkspace", "{WorkspacePath} is not in a workspace."), Args));
 		}
 	}
 }
@@ -123,7 +138,7 @@ void FPlasticSourceControlProvider::Close()
 	UserName.Empty();
 }
 
-TSharedRef<FPlasticSourceControlState, ESPMode::ThreadSafe> FPlasticSourceControlProvider::GetStateInternal(const FString& InFilename) const
+TSharedRef<FPlasticSourceControlState, ESPMode::ThreadSafe> FPlasticSourceControlProvider::GetStateInternal(const FString& InFilename)
 {
 	TSharedRef<FPlasticSourceControlState, ESPMode::ThreadSafe>* State = StateCache.Find(InFilename);
 	if (State != NULL)
@@ -134,32 +149,69 @@ TSharedRef<FPlasticSourceControlState, ESPMode::ThreadSafe> FPlasticSourceContro
 	else
 	{
 		// cache an unknown state for this item
-		TSharedRef<FPlasticSourceControlState, ESPMode::ThreadSafe> NewState = MakeShareable(new FPlasticSourceControlState(FString(InFilename)));
+		TSharedRef<FPlasticSourceControlState, ESPMode::ThreadSafe> NewState = MakeShared<FPlasticSourceControlState>(FString(InFilename));
 		StateCache.Add(InFilename, NewState);
 		return NewState;
 	}
 }
 
+TSharedRef<FPlasticSourceControlChangelistState, ESPMode::ThreadSafe> FPlasticSourceControlProvider::GetStateInternal(const FPlasticSourceControlChangelist& InChangelist)
+{
+	TSharedRef<FPlasticSourceControlChangelistState, ESPMode::ThreadSafe>* State = ChangelistsStateCache.Find(InChangelist);
+	if (State != NULL)
+	{
+		// found cached item
+		return (*State);
+	}
+	else
+	{
+		// cache an unknown state for this item
+		TSharedRef<FPlasticSourceControlChangelistState, ESPMode::ThreadSafe> NewState = MakeShared<FPlasticSourceControlChangelistState>(InChangelist);
+		ChangelistsStateCache.Add(InChangelist, NewState);
+		return NewState;
+	}
+}
+
 FText FPlasticSourceControlProvider::GetStatusText() const
 {
 	FFormatNamedArguments Args;
-    Args.Add( TEXT("PlasticScmVersion"), FText::FromString(PlasticScmVersion) );
-    Args.Add( TEXT("PluginVersion"), FText::FromString(PluginVersion) );
-	Args.Add( TEXT("WorkspacePath"), FText::FromString(PathToWorkspaceRoot) );
-	Args.Add( TEXT("WorkspaceName"), FText::FromString(WorkspaceName) );
-	Args.Add( TEXT("BranchName"), FText::FromString(BranchName) );
+	Args.Add(TEXT("PlasticScmVersion"), FText::FromString(PlasticScmVersion.String));
+	Args.Add(TEXT("PluginVersion"), FText::FromString(PluginVersion));
+	Args.Add(TEXT("WorkspacePath"), FText::FromString(PathToWorkspaceRoot));
+	Args.Add(TEXT("WorkspaceName"), FText::FromString(WorkspaceName));
+	Args.Add(TEXT("BranchName"), FText::FromString(BranchName));
 	// Detect special case for a partial checkout (CS:-1 in Gluon mode)!
 	if (-1 != ChangesetNumber)
 	{
-		Args.Add( TEXT("ChangesetNumber"), FText::FromString(FString::Printf(TEXT("%d (standard mode)"), ChangesetNumber)) );
+		Args.Add(TEXT("ChangesetNumber"), FText::FromString(FString::Printf(TEXT("%d  (regular full workspace)"), ChangesetNumber)));
 	}
 	else
 	{
-		Args.Add( TEXT("ChangesetNumber"), FText::FromString(FString::Printf(TEXT("N/A (Gluon/partial mode)"))) );
+		Args.Add(TEXT("ChangesetNumber"), FText::FromString(FString::Printf(TEXT("N/A  (Gluon/partial workspace)"))));
+	}
+	Args.Add(TEXT("UserName"), FText::FromString(UserName));
+	const FString DisplayName = PlasticSourceControlUtils::UserNameToDisplayName(UserName);
+	if (DisplayName != UserName)
+	{
+		Args.Add(TEXT("DisplayName"), FText::FromString(TEXT("(Display: ") + DisplayName + TEXT(")")));
 	}
-	Args.Add( TEXT("UserName"), FText::FromString(UserName) );
+	else
+	{
+		Args.Add(TEXT("DisplayName"), FText::GetEmpty());
+	}
+
+	FText FormattedError;
+	TArray<FString> RecentErrors = GetLastErrors();
+	if (RecentErrors.Num() > 0)
+	{
+		FFormatNamedArguments ErrorArgs;
+		ErrorArgs.Add(TEXT("ErrorText"), FText::FromString(RecentErrors[0]));
+
+		FormattedError = FText::Format(LOCTEXT("PlasticErrorStatusText", "Error: {ErrorText} {UserName}\n\n"), ErrorArgs);
+	}
+	Args.Add(TEXT("ErrorText"), FormattedError);
 
-	return FText::Format( NSLOCTEXT("Status", "Provider: Plastic\nEnabledLabel", "Plastic SCM {PlasticScmVersion} (plugin v{PluginVersion})\nWorkspace: {WorkspaceName} ({WorkspacePath})\n{BranchName}\nChangeset: {ChangesetNumber}\nUser: {UserName}"), Args );
+	return FText::Format(LOCTEXT("PlasticStatusText", "{ErrorText}Plastic SCM {PlasticScmVersion}  (plugin v{PluginVersion})\nWorkspace: {WorkspaceName}  ({WorkspacePath})\n{BranchName}\nChangeset: {ChangesetNumber}\nUser: '{UserName}'  {DisplayName}"), Args);
 }
 
 /** Quick check if source control is enabled. Specifically, it returns true if a source control provider is set (regardless of whether the provider is available) and false if no provider is set. So all providers except the stub DefaultSourceProvider will return true. */
@@ -179,7 +231,20 @@ const FName& FPlasticSourceControlProvider::GetName(void) const
 	return ProviderName;
 }
 
-ECommandResult::Type FPlasticSourceControlProvider::GetState( const TArray<FString>& InFiles, TArray< TSharedRef<ISourceControlState, ESPMode::ThreadSafe> >& OutState, EStateCacheUsage::Type InStateCacheUsage )
+void FPlasticSourceControlProvider::SetLastErrors(const TArray<FString>& InErrors)
+{
+	FScopeLock Lock(&LastErrorsCriticalSection);
+	LastErrors = InErrors;
+}
+
+TArray<FString> FPlasticSourceControlProvider::GetLastErrors() const
+{
+	FScopeLock Lock(&LastErrorsCriticalSection);
+	TArray<FString> Result = LastErrors;
+	return Result;
+}
+
+ECommandResult::Type FPlasticSourceControlProvider::GetState(const TArray<FString>& InFiles, TArray< TSharedRef<ISourceControlState, ESPMode::ThreadSafe> >& OutState, EStateCacheUsage::Type InStateCacheUsage)
 {
 	if (!IsEnabled())
 	{
@@ -204,9 +269,48 @@ ECommandResult::Type FPlasticSourceControlProvider::GetState( const TArray<FStri
 
 ECommandResult::Type FPlasticSourceControlProvider::GetState(const TArray<FSourceControlChangelistRef>& InChangelists, TArray<FSourceControlChangelistStateRef>& OutState, EStateCacheUsage::Type InStateCacheUsage)
 {
-	return ECommandResult::Failed;
+	if (!IsEnabled())
+	{
+		return ECommandResult::Failed;
+	}
+
+	if (InStateCacheUsage == EStateCacheUsage::ForceUpdate)
+	{
+		TSharedRef<class FUpdatePendingChangelistsStatus, ESPMode::ThreadSafe> UpdatePendingChangelistsOperation = ISourceControlOperation::Create<FUpdatePendingChangelistsStatus>();
+		UpdatePendingChangelistsOperation->SetChangelistsToUpdate(InChangelists);
+
+		ISourceControlProvider::Execute(UpdatePendingChangelistsOperation, EConcurrency::Synchronous);
+	}
+
+	for (FSourceControlChangelistRef Changelist : InChangelists)
+	{
+		FPlasticSourceControlChangelistRef PlasticChangelist = StaticCastSharedRef<FPlasticSourceControlChangelist>(Changelist);
+		OutState.Add(GetStateInternal(PlasticChangelist.Get()));
+	}
+
+	return ECommandResult::Succeeded;
+}
+
+bool FPlasticSourceControlProvider::RemoveChangelistFromCache(const FPlasticSourceControlChangelist& Changelist)
+{
+	return ChangelistsStateCache.Remove(Changelist) > 0;
+}
+
+TArray<FSourceControlChangelistStateRef> FPlasticSourceControlProvider::GetCachedStateByPredicate(TFunctionRef<bool(const FSourceControlChangelistStateRef&)> Predicate) const
+{
+	TArray<FSourceControlChangelistStateRef> Result;
+	for (const auto& CacheItem : ChangelistsStateCache)
+	{
+		FSourceControlChangelistStateRef State = CacheItem.Value;
+		if (Predicate(State))
+		{
+			Result.Add(State);
+		}
+	}
+	return Result;
 }
 
+
 TArray<FSourceControlStateRef> FPlasticSourceControlProvider::GetCachedStateByPredicate(TFunctionRef<bool(const FSourceControlStateRef&)> Predicate) const
 {
 	TArray<FSourceControlStateRef> Result;
@@ -226,21 +330,22 @@ bool FPlasticSourceControlProvider::RemoveFileFromCache(const FString& Filename)
 	return StateCache.Remove(Filename) > 0;
 }
 
-FDelegateHandle FPlasticSourceControlProvider::RegisterSourceControlStateChanged_Handle( const FSourceControlStateChanged::FDelegate& SourceControlStateChanged )
+FDelegateHandle FPlasticSourceControlProvider::RegisterSourceControlStateChanged_Handle(const FSourceControlStateChanged::FDelegate& SourceControlStateChanged)
 {
-	return OnSourceControlStateChanged.Add( SourceControlStateChanged );
+	return OnSourceControlStateChanged.Add(SourceControlStateChanged);
 }
 
-void FPlasticSourceControlProvider::UnregisterSourceControlStateChanged_Handle( FDelegateHandle Handle )
+void FPlasticSourceControlProvider::UnregisterSourceControlStateChanged_Handle(FDelegateHandle Handle)
 {
-	OnSourceControlStateChanged.Remove( Handle );
+	OnSourceControlStateChanged.Remove(Handle);
 }
 
-ECommandResult::Type FPlasticSourceControlProvider::Execute( const FSourceControlOperationRef& InOperation, FSourceControlChangelistPtr InChangelist, const TArray<FString>& InFiles, EConcurrency::Type InConcurrency, const FSourceControlOperationComplete& InOperationCompleteDelegate )
+ECommandResult::Type FPlasticSourceControlProvider::Execute(const FSourceControlOperationRef& InOperation, FSourceControlChangelistPtr InChangelist, const TArray<FString>& InFiles, EConcurrency::Type InConcurrency, const FSourceControlOperationComplete& InOperationCompleteDelegate)
 {
 	if (!bWorkspaceFound && !(InOperation->GetName() == "Connect") && !(InOperation->GetName() == "MakeWorkspace"))
 	{
-		UE_LOG(LogSourceControl, Warning, TEXT("'%s': only Connect operation allowed whithout a workspace"), *InOperation->GetName().ToString());
+		UE_LOG(LogSourceControl, Warning, TEXT("'%s': only Connect operation allowed without a workspace"), *InOperation->GetName().ToString());
+		InOperationCompleteDelegate.ExecuteIfBound(InOperation, ECommandResult::Failed);
 		return ECommandResult::Failed;
 	}
 
@@ -250,8 +355,8 @@ ECommandResult::Type FPlasticSourceControlProvider::Execute( const FSourceContro
 	{
 		// this operation is unsupported by this source control provider
 		FFormatNamedArguments Arguments;
-		Arguments.Add( TEXT("OperationName"), FText::FromName(InOperation->GetName()) );
-		Arguments.Add( TEXT("ProviderName"), FText::FromName(GetName()) );
+		Arguments.Add(TEXT("OperationName"), FText::FromName(InOperation->GetName()));
+		Arguments.Add(TEXT("ProviderName"), FText::FromName(GetName()));
 		FMessageLog("SourceControl").Error(FText::Format(LOCTEXT("UnsupportedOperation", "Operation '{OperationName}' not supported by source control provider '{ProviderName}'"), Arguments));
 		return ECommandResult::Failed;
 	}
@@ -260,6 +365,9 @@ ECommandResult::Type FPlasticSourceControlProvider::Execute( const FSourceContro
 	Command->Files = SourceControlHelpers::AbsoluteFilenames(InFiles);
 	Command->OperationCompleteDelegate = InOperationCompleteDelegate;
 
+	TSharedPtr<FPlasticSourceControlChangelist, ESPMode::ThreadSafe> ChangelistPtr = StaticCastSharedPtr<FPlasticSourceControlChangelist>(InChangelist);
+	Command->Changelist = ChangelistPtr ? ChangelistPtr.ToSharedRef().Get() : FPlasticSourceControlChangelist();
+
 	// fire off operation
 	if (InConcurrency == EConcurrency::Synchronous)
 	{
@@ -277,12 +385,12 @@ ECommandResult::Type FPlasticSourceControlProvider::Execute( const FSourceContro
 	}
 }
 
-bool FPlasticSourceControlProvider::CanCancelOperation( const FSourceControlOperationRef& InOperation ) const
+bool FPlasticSourceControlProvider::CanCancelOperation(const FSourceControlOperationRef& InOperation) const
 {
 	return false;
 }
 
-void FPlasticSourceControlProvider::CancelOperation( const FSourceControlOperationRef& InOperation )
+void FPlasticSourceControlProvider::CancelOperation(const FSourceControlOperationRef& InOperation)
 {
 }
 
@@ -293,28 +401,28 @@ bool FPlasticSourceControlProvider::UsesLocalReadOnlyState() const
 
 bool FPlasticSourceControlProvider::UsesChangelists() const
 {
-	return false; // We don't want to show ChangeList column anymore (Plastic SCM term would be ChangeSet)
+	return true;
 }
 
 bool FPlasticSourceControlProvider::UsesCheckout() const
 {
-	return true;
+	return GetDefault<UPlasticSourceControlProjectSettings>()->bPromptForCheckoutOnChange;
 }
 
-TSharedPtr<IPlasticSourceControlWorker, ESPMode::ThreadSafe> FPlasticSourceControlProvider::CreateWorker(const FName& InOperationName) const
+TSharedPtr<IPlasticSourceControlWorker, ESPMode::ThreadSafe> FPlasticSourceControlProvider::CreateWorker(const FName& InOperationName)
 {
 	const FGetPlasticSourceControlWorker* Operation = WorkersMap.Find(InOperationName);
 	if (Operation != nullptr)
 	{
-		return Operation->Execute();
+		return Operation->Execute(*this);
 	}
-		
+
 	return nullptr;
 }
 
-void FPlasticSourceControlProvider::RegisterWorker( const FName& InName, const FGetPlasticSourceControlWorker& InDelegate )
+void FPlasticSourceControlProvider::RegisterWorker(const FName& InName, const FGetPlasticSourceControlWorker& InDelegate)
 {
-	WorkersMap.Add( InName, InDelegate );
+	WorkersMap.Add(InName, InDelegate);
 }
 
 void FPlasticSourceControlProvider::OutputCommandMessages(const FPlasticSourceControlCommand& InCommand) const
@@ -344,17 +452,50 @@ void FPlasticSourceControlProvider::UpdateWorkspaceStatus(const class FPlasticSo
 		RepositoryName = InCommand.RepositoryName;
 		ServerUrl = InCommand.ServerUrl;
 
+		// only pop-up errors when running in full Editor, not in command line scripts
+		if (!IsRunningCommandlet())
+		{
+			if (bPlasticAvailable)
+			{
+				if (PlasticScmVersion < PlasticSourceControlUtils::GetOldestSupportedPlasticScmVersion())
+				{
+					FFormatNamedArguments Args;
+					Args.Add(TEXT("PlasticScmVersion"), FText::FromString(PlasticScmVersion.String));
+					Args.Add(TEXT("OldestSupportedPlasticScmVersion"), FText::FromString(PlasticSourceControlUtils::GetOldestSupportedPlasticScmVersion().String));
+					const FText UnsuportedVersionWarning = FText::Format(LOCTEXT("Plastic_UnsuportedVersion", "Plastic SCM {PlasticScmVersion} is not supported anymore by this plugin.\nPlastic SCM {OldestSupportedPlasticScmVersion} or a more recent version is required.\nPlease upgrade to the latest version."), Args);
+					FMessageLog("SourceControl").Warning(UnsuportedVersionWarning);
+					FMessageDialog::Open(EAppMsgType::Ok, UnsuportedVersionWarning);
+				}
+			}
+			else if (InCommand.ErrorMessages.Num() > 0)
+			{
+				FMessageDialog::Open(EAppMsgType::Ok, FText::FromString(InCommand.ErrorMessages[0]));
+			}
+		}
+
 		if (bWorkspaceFound)
 		{
 			// Extend the "Source Control" menu in the Editor Toolbar on each successful connection
 			PlasticSourceControlMenu.Unregister(); // cleanup for any previous connection
 			PlasticSourceControlMenu.Register();
 		}
+
+		SetLastErrors(InCommand.ErrorMessages);
 	}
 	else if (InCommand.bConnectionDropped)
 	{
 		// checkconnection failed on UpdateStatus
 		bServerAvailable = false;
+
+		SetLastErrors(InCommand.ErrorMessages);
+	}
+	else if (!bServerAvailable)
+	{
+		bServerAvailable = InCommand.bCommandSuccessful;
+		if (bServerAvailable)
+		{
+			SetLastErrors(TArray<FString>());
+		}
 	}
 
 	// And for all operations running UpdateStatus, get Changeset and Branch informations:
@@ -369,7 +510,9 @@ void FPlasticSourceControlProvider::UpdateWorkspaceStatus(const class FPlasticSo
 }
 
 void FPlasticSourceControlProvider::Tick()
-{	
+{
+	TRACE_CPUPROFILER_EVENT_SCOPE(FPlasticSourceControlProvider::Tick);
+
 	bool bStatesUpdated = false;
 	for (int32 CommandIndex = 0; CommandIndex < CommandQueue.Num(); ++CommandIndex)
 	{
@@ -408,7 +551,7 @@ void FPlasticSourceControlProvider::Tick()
 				delete &Command;
 			}
 
-			// only do one command per tick loop, as we dont want concurrent modification 
+			// only do one command per tick loop, as we don't want concurrent modification
 			// of the command queue (which can happen in the completion delegate)
 			break;
 		}
@@ -416,11 +559,12 @@ void FPlasticSourceControlProvider::Tick()
 
 	if (bStatesUpdated)
 	{
+		TRACE_CPUPROFILER_EVENT_SCOPE(FPlasticSourceControlProvider::Tick::BroadcastStateUpdate);
 		OnSourceControlStateChanged.Broadcast();
 	}
 }
 
-TArray< TSharedRef<ISourceControlLabel> > FPlasticSourceControlProvider::GetLabels( const FString& InMatchingSpec ) const
+TArray<TSharedRef<ISourceControlLabel>> FPlasticSourceControlProvider::GetLabels(const FString& InMatchingSpec) const
 {
 	TArray< TSharedRef<ISourceControlLabel> > Tags;
 	return Tags;
@@ -428,7 +572,22 @@ TArray< TSharedRef<ISourceControlLabel> > FPlasticSourceControlProvider::GetLabe
 
 TArray<FSourceControlChangelistRef> FPlasticSourceControlProvider::GetChangelists( EStateCacheUsage::Type InStateCacheUsage )
 {
-	return TArray<FSourceControlChangelistRef>();
+	if (!IsEnabled())
+	{
+		return TArray<FSourceControlChangelistRef>();
+	}
+
+	if (InStateCacheUsage == EStateCacheUsage::ForceUpdate)
+	{
+		TSharedRef<class FUpdatePendingChangelistsStatus, ESPMode::ThreadSafe> UpdatePendingChangelistsOperation = ISourceControlOperation::Create<FUpdatePendingChangelistsStatus>();
+		UpdatePendingChangelistsOperation->SetUpdateAllChangelists(true);
+
+		ISourceControlProvider::Execute(UpdatePendingChangelistsOperation, EConcurrency::Synchronous);
+	}
+
+	TArray<FSourceControlChangelistRef> Changelists;
+	Algo::Transform(ChangelistsStateCache, Changelists, [](const auto& Pair) { return MakeShared<FPlasticSourceControlChangelist, ESPMode::ThreadSafe>(Pair.Key); });
+	return Changelists;
 }
 
 #if SOURCE_CONTROL_WITH_SLATE
@@ -440,6 +599,8 @@ TSharedRef<class SWidget> FPlasticSourceControlProvider::MakeSettingsWidget() co
 
 ECommandResult::Type FPlasticSourceControlProvider::ExecuteSynchronousCommand(FPlasticSourceControlCommand& InCommand, const FText& Task)
 {
+	TRACE_CPUPROFILER_EVENT_SCOPE(FPlasticSourceControlProvider::ExecuteSynchronousCommand);
+
 	ECommandResult::Type Result = ECommandResult::Failed;
 
 	// Display the progress dialog if a string was provided
@@ -447,7 +608,7 @@ ECommandResult::Type FPlasticSourceControlProvider::ExecuteSynchronousCommand(FP
 		FScopedSourceControlProgress Progress(Task);
 
 		// Issue the command asynchronously...
-		IssueCommand( InCommand );
+		IssueCommand(InCommand);
 
 		// ... then wait for its completion (thus making it synchronous)
 		while (!InCommand.bExecuteProcessed)
@@ -471,7 +632,7 @@ ECommandResult::Type FPlasticSourceControlProvider::ExecuteSynchronousCommand(FP
 		else
 		{
 			// TODO If the command failed, inform the user that they need to try again (see Perforce)
-			//FMessageDialog::Open( EAppMsgType::Ok, LOCTEXT("Plastic_ServerUnresponsive", "Plastic server is unresponsive. Please check your connection and try again.") );
+			// FMessageDialog::Open( EAppMsgType::Ok, LOCTEXT("Plastic_ServerUnresponsive", "Plastic server is unresponsive. Please check your connection and try again.") );
 
 			UE_LOG(LogSourceControl, Error, TEXT("Command '%s' Failed!"), *InCommand.Operation->GetName().ToString());
 		}
@@ -481,9 +642,9 @@ ECommandResult::Type FPlasticSourceControlProvider::ExecuteSynchronousCommand(FP
 	check(!InCommand.bAutoDelete);
 
 	// ensure commands that are not auto deleted do not end up in the command queue
-	if ( CommandQueue.Contains( &InCommand ) ) 
+	if (CommandQueue.Contains(&InCommand))
 	{
-		CommandQueue.Remove( &InCommand );
+		CommandQueue.Remove(&InCommand);
 	}
 	delete &InCommand;
 
diff --git a/Engine/Plugins/Developer/PlasticSourceControl/Source/PlasticSourceControl/Private/PlasticSourceControlProvider.h b/Engine/Plugins/Developer/PlasticSourceControl/Source/PlasticSourceControl/Private/PlasticSourceControlProvider.h
--- a/Engine/Plugins/Developer/PlasticSourceControl/Source/PlasticSourceControl/Private/PlasticSourceControlProvider.h
+++ b/Engine/Plugins/Developer/PlasticSourceControl/Source/PlasticSourceControl/Private/PlasticSourceControlProvider.h
@@ -7,22 +7,19 @@
 #include "IPlasticSourceControlWorker.h"
 #include "PlasticSourceControlConsole.h"
 #include "PlasticSourceControlMenu.h"
+#include "PlasticSourceControlSettings.h"
+#include "SoftwareVersion.h"
 
-class FPlasticSourceControlState;
+#include "ISourceControlChangelistState.h"
+#include "PlasticSourceControlChangelist.h"
 
-DECLARE_DELEGATE_RetVal(FPlasticSourceControlWorkerRef, FGetPlasticSourceControlWorker)
+DECLARE_DELEGATE_RetVal_OneParam(FPlasticSourceControlWorkerRef, FGetPlasticSourceControlWorker, FPlasticSourceControlProvider&)
 
 class FPlasticSourceControlProvider : public ISourceControlProvider
 {
 public:
 	/** Constructor */
-	FPlasticSourceControlProvider()
-		: bPlasticAvailable(false)
-		, bWorkspaceFound(false)
-		, bServerAvailable(false)
-		, ChangesetNumber(0)
-	{
-	}
+	FPlasticSourceControlProvider();
 
 	/* ISourceControlProvider implementation */
 	virtual void Init(bool bForceConnection = true) override;
@@ -113,8 +110,29 @@ public:
 		return ChangesetNumber;
 	}
 
+	/** Version of the Plastic SCM executable used */
+	inline const FSoftwareVersion& GetPlasticScmVersion() const
+	{
+		return PlasticScmVersion;
+	}
+
+	/** Version of the Plastic SCM plugin */
+	const FString& GetPluginVersion() const
+	{
+		return PluginVersion;
+	}
+
+	/** Set list of error messages that occurred after last Plastic command */
+	void SetLastErrors(const TArray<FString>& InErrors);
+
+	/** Get list of error messages that occurred after last Plastic command */
+	TArray<FString> GetLastErrors() const;
+
 	/** Helper function used to update state cache */
-	TSharedRef<FPlasticSourceControlState, ESPMode::ThreadSafe> GetStateInternal(const FString& InFilename) const;
+	TSharedRef<class FPlasticSourceControlState, ESPMode::ThreadSafe> GetStateInternal(const FString& InFilename);
+
+	/** Helper function used to update changelists state cache */
+	TSharedRef<class FPlasticSourceControlChangelistState, ESPMode::ThreadSafe> GetStateInternal(const FPlasticSourceControlChangelist& InChangelist);
 
 	/**
 	 * Register a worker with the provider.
@@ -125,19 +143,46 @@ public:
 	/** Remove a named file from the state cache */
 	bool RemoveFileFromCache(const FString& Filename);
 
-private:
+	/** Remove a changelist from the state cache */
+	bool RemoveChangelistFromCache(const FPlasticSourceControlChangelist& Changelist);
+
+	/** Returns a list of changelists from the cache based on a given predicate */
+	TArray<FSourceControlChangelistStateRef> GetCachedStateByPredicate(TFunctionRef<bool(const FSourceControlChangelistStateRef&)> Predicate) const;
+
+	/** Access the Plastic source control settings */
+	FPlasticSourceControlSettings& AccessSettings()
+	{
+		return PlasticSourceControlSettings;
+	}
+	const FPlasticSourceControlSettings& AccessSettings() const
+	{
+		return PlasticSourceControlSettings;
+	}
 
+	/** Save the Plastic source control settings */
+	void SaveSettings()
+	{
+		PlasticSourceControlSettings.SaveSettings();
+	}
+
+private:
 	/** Is Plastic binary found and working. */
-	bool bPlasticAvailable;
+	bool bPlasticAvailable = false;
 
 	/** Is Plastic workspace found. */
-	bool bWorkspaceFound;
+	bool bWorkspaceFound = false;
 
 	/** Indicates if source control integration is available or not. */
-	bool bServerAvailable;
+	bool bServerAvailable = false;
+
+	/** Critical section for thread safety of error messages that occurred after last Plastic command */
+	mutable FCriticalSection LastErrorsCriticalSection;
+
+	/** List of error messages that occurred after last Plastic command */
+	TArray<FString> LastErrors;
 
 	/** Helper function for Execute() */
-	TSharedPtr<class IPlasticSourceControlWorker, ESPMode::ThreadSafe> CreateWorker(const FName& InOperationName) const;
+	TSharedPtr<class IPlasticSourceControlWorker, ESPMode::ThreadSafe> CreateWorker(const FName& InOperationName);
 
 	/** Helper function for running command synchronously. */
 	ECommandResult::Type ExecuteSynchronousCommand(class FPlasticSourceControlCommand& InCommand, const FText& Task);
@@ -150,8 +195,8 @@ private:
 	/** Update workspace status on Connect and UpdateStatus operations */
 	void UpdateWorkspaceStatus(const class FPlasticSourceControlCommand& InCommand);
 
-    /** Version of the Plastic SCM executable used */
-    FString PlasticScmVersion;
+	/** Version of the Plastic SCM executable used */
+	FSoftwareVersion PlasticScmVersion;
 
 	/** Version of the Plastic SCM plugin */
 	FString PluginVersion;
@@ -175,10 +220,11 @@ private:
 	FString BranchName;
 
 	/** Current Changeset Number */
-	int32 ChangesetNumber;
+	int32 ChangesetNumber = 0;
 
-	/** State cache */
-	mutable TMap<FString, TSharedRef<FPlasticSourceControlState, ESPMode::ThreadSafe> > StateCache;
+	/** State caches */
+	TMap<FString, TSharedRef<class FPlasticSourceControlState, ESPMode::ThreadSafe> > StateCache;
+	TMap<FPlasticSourceControlChangelist, TSharedRef<class FPlasticSourceControlChangelistState, ESPMode::ThreadSafe> > ChangelistsStateCache;
 
 	/** The currently registered source control operations */
 	TMap<FName, FGetPlasticSourceControlWorker> WorkersMap;
@@ -189,9 +235,12 @@ private:
 	/** For notifying when the source control states in the cache have changed */
 	FSourceControlStateChanged OnSourceControlStateChanged;
 
+	/** Source Control Console commands */
+	FPlasticSourceControlConsole PlasticSourceControlConsole;
+
 	/** Source Control Menu Extension */
 	FPlasticSourceControlMenu PlasticSourceControlMenu;
 
-	/** Source Control Console commands */
-	FPlasticSourceControlConsole PlasticSourceControlConsole;
+	/** The settings for Plastic source control */
+	FPlasticSourceControlSettings PlasticSourceControlSettings;
 };
diff --git a/Engine/Plugins/Developer/PlasticSourceControl/Source/PlasticSourceControl/Private/PlasticSourceControlRevision.cpp b/Engine/Plugins/Developer/PlasticSourceControl/Source/PlasticSourceControl/Private/PlasticSourceControlRevision.cpp
--- a/Engine/Plugins/Developer/PlasticSourceControl/Source/PlasticSourceControl/Private/PlasticSourceControlRevision.cpp
+++ b/Engine/Plugins/Developer/PlasticSourceControl/Source/PlasticSourceControl/Private/PlasticSourceControlRevision.cpp
@@ -20,13 +20,6 @@ bool FPlasticSourceControlRevision::Get(FString& InOutFilename, EConcurrency::Ty
 		UE_LOG(LogSourceControl, Warning, TEXT("Only EConcurrency::Synchronous is tested/supported for this operation."));
 	}
 
-	const FPlasticSourceControlModule& PlasticSourceControl = FModuleManager::LoadModuleChecked<FPlasticSourceControlModule>("PlasticSourceControl");
-	const FString& PathToPlasticBinary = PlasticSourceControl.AccessSettings().GetBinaryPath();
-	if (!State)
-	{
-		UE_LOG(LogSourceControl, Error, TEXT("Revision(%s %d): unknown state!"), *Filename, RevisionId);
-	}
-
 	// if a filename for the temp file wasn't supplied generate a unique-ish one
 	if (InOutFilename.Len() == 0)
 	{
@@ -44,12 +37,15 @@ bool FPlasticSourceControlRevision::Get(FString& InOutFilename, EConcurrency::Ty
 	}
 	else if (State)
 	{
+		const FString& PathToPlasticBinary = FPlasticSourceControlModule::Get().GetProvider().AccessSettings().GetBinaryPath();
+
 		// Format the revision specification of the file, like rev:Content/BP.uasset#cs:12@repo@server:8087
 		const FString RevisionSpecification = FString::Printf(TEXT("rev:%s#cs:%d@%s"), *Filename, ChangesetNumber, *State->RepSpec);
 		bCommandSuccessful = PlasticSourceControlUtils::RunDumpToFile(PathToPlasticBinary, RevisionSpecification, InOutFilename);
 	}
 	else
 	{
+		UE_LOG(LogSourceControl, Error, TEXT("Revision(%s %d): unknown state!"), *Filename, RevisionId);
 		bCommandSuccessful = false;
 	}
 	return bCommandSuccessful;
diff --git a/Engine/Plugins/Developer/PlasticSourceControl/Source/PlasticSourceControl/Private/PlasticSourceControlRevision.h b/Engine/Plugins/Developer/PlasticSourceControl/Source/PlasticSourceControl/Private/PlasticSourceControlRevision.h
--- a/Engine/Plugins/Developer/PlasticSourceControl/Source/PlasticSourceControl/Private/PlasticSourceControlRevision.h
+++ b/Engine/Plugins/Developer/PlasticSourceControl/Source/PlasticSourceControl/Private/PlasticSourceControlRevision.h
@@ -8,7 +8,7 @@
 class FPlasticSourceControlState;
 
 /** Revision of a file, linked to a specific commit */
-class FPlasticSourceControlRevision : public ISourceControlRevision, public TSharedFromThis<FPlasticSourceControlRevision, ESPMode::ThreadSafe>
+class FPlasticSourceControlRevision : public ISourceControlRevision
 {
 public:
 	FPlasticSourceControlRevision()
@@ -21,8 +21,8 @@ public:
 
 	/** ISourceControlRevision interface */
 	virtual bool Get(FString& InOutFilename, EConcurrency::Type InConcurrency = EConcurrency::Synchronous) const override;
-	virtual bool GetAnnotated( TArray<FAnnotationLine>& OutLines ) const override;
-	virtual bool GetAnnotated( FString& InOutFilename ) const override;
+	virtual bool GetAnnotated(TArray<FAnnotationLine>& OutLines) const override;
+	virtual bool GetAnnotated(FString& InOutFilename) const override;
 	virtual const FString& GetFilename() const override;
 	virtual int32 GetRevisionNumber() const override;
 	virtual const FString& GetRevision() const override;
@@ -36,7 +36,6 @@ public:
 	virtual int32 GetFileSize() const override;
 
 public:
-
 	/** Point back to State this Revision is from */
 	FPlasticSourceControlState* State = nullptr;
 
diff --git a/Engine/Plugins/Developer/PlasticSourceControl/Source/PlasticSourceControl/Private/PlasticSourceControlState.cpp b/Engine/Plugins/Developer/PlasticSourceControl/Source/PlasticSourceControl/Private/PlasticSourceControlState.cpp
--- a/Engine/Plugins/Developer/PlasticSourceControl/Source/PlasticSourceControl/Private/PlasticSourceControlState.cpp
+++ b/Engine/Plugins/Developer/PlasticSourceControl/Source/PlasticSourceControl/Private/PlasticSourceControlState.cpp
@@ -1,6 +1,7 @@
 // Copyright Epic Games, Inc. All Rights Reserved.
 
 #include "PlasticSourceControlState.h"
+#include "PlasticSourceControlProjectSettings.h"
 #include "ISourceControlModule.h"
 #include "Styling/AppStyle.h"
 
@@ -31,7 +32,7 @@ const TCHAR* ToString(EWorkspaceState::Type InWorkspaceState)
 	}
 	return WorkspaceStateStr;
 }
-}
+} // namespace EWorkspaceState
 
 
 int32 FPlasticSourceControlState::GetHistorySize() const
@@ -39,13 +40,13 @@ int32 FPlasticSourceControlState::GetHistorySize() const
 	return History.Num();
 }
 
-TSharedPtr<class ISourceControlRevision, ESPMode::ThreadSafe> FPlasticSourceControlState::GetHistoryItem( int32 HistoryIndex ) const
+TSharedPtr<class ISourceControlRevision, ESPMode::ThreadSafe> FPlasticSourceControlState::GetHistoryItem(int32 HistoryIndex) const
 {
 	check(History.IsValidIndex(HistoryIndex));
 	return History[HistoryIndex];
 }
 
-TSharedPtr<class ISourceControlRevision, ESPMode::ThreadSafe> FPlasticSourceControlState::FindHistoryRevision( int32 RevisionNumber ) const
+TSharedPtr<class ISourceControlRevision, ESPMode::ThreadSafe> FPlasticSourceControlState::FindHistoryRevision(int32 RevisionNumber) const
 {
 	for (const auto& Revision : History)
 	{
@@ -75,7 +76,7 @@ TSharedPtr<class ISourceControlRevision, ESPMode::ThreadSafe> FPlasticSourceCont
 {
 	for (const auto& Revision : History)
 	{
-		// look for the the SHA1 id of the file, not the commit id (revision)
+		// look for the SHA1 id of the file, not the commit id (revision)
 		if (Revision->ChangesetNumber == PendingMergeBaseChangeset)
 		{
 			return Revision;
@@ -184,7 +185,7 @@ FText FPlasticSourceControlState::GetDisplayTooltip() const
 {
 	if (!IsCurrent())
 	{
-		return LOCTEXT("NotCurrent_Tooltip", "Not at the head revision");
+		return FText::Format(LOCTEXT("NotCurrent_Tooltip", "Not at the head revision CS:{0} by {1}"), FText::AsNumber(DepotRevisionChangeset), FText::FromString(HeadUserName));
 	}
 	else if (WorkspaceState != EWorkspaceState::LockedByOther)
 	{
@@ -199,33 +200,33 @@ FText FPlasticSourceControlState::GetDisplayTooltip() const
 	switch (WorkspaceState)
 	{
 	case EWorkspaceState::Unknown:
-		return LOCTEXT("Unknown_Tooltip", "The file status is unknown");
+		return FText();
 	case EWorkspaceState::Ignored:
-		return LOCTEXT("Ignored_Tooltip", "The file is ignored");
+		return LOCTEXT("Ignored_Tooltip", "Ignored");
 	case EWorkspaceState::Controlled:
-		return LOCTEXT("Pristine_Tooltip", "The file has no modification");
+		return FText();
 	case EWorkspaceState::CheckedOut:
-		return LOCTEXT("CheckedOut_Tooltip", "The file is checked out");
+		return LOCTEXT("CheckedOut_Tooltip", "Checked-out");
 	case EWorkspaceState::Added:
-		return LOCTEXT("Added_Tooltip", "The file has been added");
+		return LOCTEXT("Added_Tooltip", "Added");
 	case EWorkspaceState::Moved:
-		return LOCTEXT("Moved_Tooltip", "The file has been moved or renamed");
+		return LOCTEXT("Moved_Tooltip", "Moved or renamed");
 	case EWorkspaceState::Copied:
-		return LOCTEXT("Copied_Tooltip", "The file has been copied");
+		return LOCTEXT("Copied_Tooltip", "Copied");
 	case EWorkspaceState::Replaced:
-		return LOCTEXT("Replaced_Tooltip", "The file has been replaced / merged");
+		return LOCTEXT("Replaced_Tooltip", "Replaced: merge conflict resolved");
 	case EWorkspaceState::Deleted:
-		return LOCTEXT("Deleted_Tooltip", "The file has been deleted");
+		return LOCTEXT("Deleted_Tooltip", "Deleted");
 	case EWorkspaceState::LocallyDeleted:
-		return LOCTEXT("LocallyDeleted_Tooltip", "The file is missing");
+		return LOCTEXT("LocallyDeleted_Tooltip", "Locally Deleted");
 	case EWorkspaceState::Changed:
-		return LOCTEXT("Modified_Tooltip", "The file has been modified");
+		return LOCTEXT("Modified_Tooltip", "Locally modified");
 	case EWorkspaceState::Conflicted:
-		return LOCTEXT("ContentsConflict_Tooltip", "The content of the file conflict with updates received from the repository");
+		return LOCTEXT("ContentsConflict_Tooltip", "Conflict with updates received from the repository");
 	case EWorkspaceState::LockedByOther:
 		return FText::Format(LOCTEXT("CheckedOutOther_Tooltip", "Checked out by {0} in {1}"), FText::FromString(LockedBy), FText::FromString(LockedWhere));
 	case EWorkspaceState::Private:
-		return LOCTEXT("NotControlled_Tooltip", "The file is not under version control");
+		return LOCTEXT("NotControlled_Tooltip", "Private: not under version control");
 	}
 
 	return FText();
@@ -249,15 +250,15 @@ const FDateTime& FPlasticSourceControlState::GetTimeStamp() const
 bool FPlasticSourceControlState::CanCheckIn() const
 {
 	const bool bCanCheckIn =
-		(  WorkspaceState == EWorkspaceState::Added
+		  (WorkspaceState == EWorkspaceState::Added
 		|| WorkspaceState == EWorkspaceState::Deleted
 		|| WorkspaceState == EWorkspaceState::LocallyDeleted
 		|| WorkspaceState == EWorkspaceState::Changed
 		|| WorkspaceState == EWorkspaceState::Moved
 		|| WorkspaceState == EWorkspaceState::Copied
 		|| WorkspaceState == EWorkspaceState::Replaced
-		|| WorkspaceState == EWorkspaceState::CheckedOut
-		) && !IsConflicted() && IsCurrent();
+		|| WorkspaceState == EWorkspaceState::CheckedOut)
+		&& !IsConflicted() && IsCurrent();
 
 	if (!IsUnknown())
 	{
@@ -269,9 +270,14 @@ bool FPlasticSourceControlState::CanCheckIn() const
 
 bool FPlasticSourceControlState::CanCheckout() const
 {
-	const bool bCanCheckout  = (   WorkspaceState == EWorkspaceState::Controlled	// In source control, Unmodified
+	if (!GetDefault<UPlasticSourceControlProjectSettings>()->bPromptForCheckoutOnChange)
+	{
+		return false;
+	}
+
+	const bool bCanCheckout  =    (WorkspaceState == EWorkspaceState::Controlled	// In source control, Unmodified
 								|| WorkspaceState == EWorkspaceState::Changed		// In source control, but not checked-out
-								|| WorkspaceState == EWorkspaceState::Replaced)		// In source control, merged, waiting for checkin to conclude the merge 
+								|| WorkspaceState == EWorkspaceState::Replaced)		// In source control, merged, waiting for checkin to conclude the merge
 								&& IsCurrent(); // Is up to date (at the revision of the repo)
 
 	if (!IsUnknown())
@@ -287,7 +293,7 @@ bool FPlasticSourceControlState::IsCheckedOut() const
 	const bool bIsCheckedOut = WorkspaceState == EWorkspaceState::CheckedOut
 							|| WorkspaceState == EWorkspaceState::Moved
 							|| WorkspaceState == EWorkspaceState::Conflicted	// In source control, waiting for merged
-							|| WorkspaceState == EWorkspaceState::Replaced		// In source control, merged, waiting for checkin to conclude the merge 
+							|| WorkspaceState == EWorkspaceState::Replaced		// In source control, merged, waiting for checkin to conclude the merge
 							|| WorkspaceState == EWorkspaceState::Changed;		// Note: Workaround to enable checkin (still required by UE5.0)
 
 	if (bIsCheckedOut)
@@ -295,7 +301,15 @@ bool FPlasticSourceControlState::IsCheckedOut() const
 		UE_LOG(LogSourceControl, Verbose, TEXT("%s IsCheckedOut"), *LocalFilename);
 	}
 
-	return bIsCheckedOut;
+	if (GetDefault<UPlasticSourceControlProjectSettings>()->bPromptForCheckoutOnChange)
+	{
+		return bIsCheckedOut;
+	}
+	else
+	{
+		// Any controlled state will be considered as checked out if the prompt is disabled
+		return IsSourceControlled();
+	}
 }
 
 bool FPlasticSourceControlState::IsCheckedOutOther(FString* Who) const
@@ -318,7 +332,7 @@ bool FPlasticSourceControlState::IsCheckedOutOther(FString* Who) const
 /** Get whether this file is checked out in a different branch, if no branch is specified defaults to FEngineVerion current branch */
 bool FPlasticSourceControlState::IsCheckedOutInOtherBranch(const FString& CurrentBranch /* = FString() */) const
 {
-	// Note: to my knowledge, it's not possible to detect that with PlasticSCM withouth the Locks,
+	// Note: to my knowledge, it's not possible to detect that with PlasticSCM without the Locks,
 	// which are already detected by fileinfo LockedBy/LockedWhere and reported by IsCheckedOutOther() above
 	return false;
 }
@@ -350,7 +364,7 @@ bool FPlasticSourceControlState::IsCurrent() const
 	{
 		UE_LOG(LogSourceControl, VeryVerbose, TEXT("%s IsCurrent"), *LocalFilename);
 	}
-	
+
 	return bIsCurrent;
 }
 
@@ -417,9 +431,9 @@ bool FPlasticSourceControlState::IsUnknown() const
 
 bool FPlasticSourceControlState::IsModified() const
 {
-	// Warning: for a clean "check-in" (commit) checked-out files unmodified should be removed from the changeset (the index)
+	// Warning: for a clean "checkin" (commit) checked-out files unmodified should be removed from the changeset (the index)
 	//
-	// Thus, before check-in UE4 Editor call RevertUnchangedFiles() in PromptForCheckin() and CheckinFiles().
+	// Thus, before checkin UE4 Editor call RevertUnchangedFiles() in PromptForCheckin() and CheckinFiles().
 	//
 	// So here we must take care to enumerate all states that need to be commited, all other will be discarded:
 	//  - Unknown
diff --git a/Engine/Plugins/Developer/PlasticSourceControl/Source/PlasticSourceControl/Private/PlasticSourceControlState.h b/Engine/Plugins/Developer/PlasticSourceControl/Source/PlasticSourceControl/Private/PlasticSourceControlState.h
--- a/Engine/Plugins/Developer/PlasticSourceControl/Source/PlasticSourceControl/Private/PlasticSourceControlState.h
+++ b/Engine/Plugins/Developer/PlasticSourceControl/Source/PlasticSourceControl/Private/PlasticSourceControlState.h
@@ -4,7 +4,9 @@
 
 #include "CoreMinimal.h"
 #include "ISourceControlState.h"
+#include "ISourceControlRevision.h"
 #include "PlasticSourceControlRevision.h"
+#include "PlasticSourceControlChangelist.h"
 
 namespace EWorkspaceState
 {
@@ -28,18 +30,13 @@ namespace EWorkspaceState
 
 	// debug log utility
 	const TCHAR* ToString(EWorkspaceState::Type InWorkspaceState);
-}
+} // namespace EWorkspaceState
 
-class FPlasticSourceControlState : public ISourceControlState, public TSharedFromThis<FPlasticSourceControlState, ESPMode::ThreadSafe>
+class FPlasticSourceControlState : public ISourceControlState
 {
 public:
-
-	FPlasticSourceControlState(FString&& InLocalFilename)
+	explicit FPlasticSourceControlState(FString&& InLocalFilename)
 		: LocalFilename(MoveTemp(InLocalFilename))
-		, WorkspaceState(EWorkspaceState::Unknown)
-		, DepotRevisionChangeset(INVALID_REVISION)
-		, LocalRevisionChangeset(INVALID_REVISION)
-		, TimeStamp(0)
 	{
 	}
 
@@ -48,10 +45,6 @@ public:
 	const FPlasticSourceControlState& operator=(const FPlasticSourceControlState& InState) = delete;
 
 	FPlasticSourceControlState(FPlasticSourceControlState&& InState)
-		: WorkspaceState(EWorkspaceState::Unknown)
-		, DepotRevisionChangeset(INVALID_REVISION)
-		, LocalRevisionChangeset(INVALID_REVISION)
-		, TimeStamp(0)
 	{
 		Move(MoveTemp(InState));
 	}
@@ -156,10 +149,10 @@ public:
 	FString PendingMergeFilename;
 
 	/** Changeset with which our local revision diverged from the source/remote revision */
-	int32 PendingMergeBaseChangeset;
+	int32 PendingMergeBaseChangeset = INVALID_REVISION;
 
 	/** Changeset of the source/remote revision of the merge in progress */
-	int32 PendingMergeSourceChangeset;
+	int32 PendingMergeSourceChangeset = INVALID_REVISION;
 
 	/** Plastic SCM Parameters of the merge in progress */
 	TArray<FString> PendingMergeParameters;
@@ -171,19 +164,22 @@ public:
 	FString LockedWhere;
 
 	/** State of the workspace */
-	EWorkspaceState::Type WorkspaceState;
+	EWorkspaceState::Type WorkspaceState = EWorkspaceState::Unknown;
 
 	/** Latest revision number of the file in the depot */
-	int DepotRevisionChangeset;
+	int DepotRevisionChangeset = INVALID_REVISION;
 
 	/** Latest revision number at which a file was synced to before being edited */
-	int LocalRevisionChangeset;
+	int LocalRevisionChangeset = INVALID_REVISION;
 
 	/** Original name in case of a Moved/Renamed file */
 	FString MovedFrom;
 
+	/** Changelist containing this file */
+	FPlasticSourceControlChangelist Changelist;
+
 	/** The timestamp of the last update */
-	FDateTime TimeStamp;
+	FDateTime TimeStamp = 0;
 
 	/** The branch with the head change list */
 	FString HeadBranch;
diff --git a/Engine/Plugins/Developer/PlasticSourceControl/Source/PlasticSourceControl/Private/PlasticSourceControlUtils.cpp b/Engine/Plugins/Developer/PlasticSourceControl/Source/PlasticSourceControl/Private/PlasticSourceControlUtils.cpp
--- a/Engine/Plugins/Developer/PlasticSourceControl/Source/PlasticSourceControl/Private/PlasticSourceControlUtils.cpp
+++ b/Engine/Plugins/Developer/PlasticSourceControl/Source/PlasticSourceControl/Private/PlasticSourceControlUtils.cpp
@@ -3,6 +3,8 @@
 #include "PlasticSourceControlUtils.h"
 #include "PlasticSourceControlCommand.h"
 #include "PlasticSourceControlModule.h"
+#include "PlasticSourceControlProjectSettings.h"
+#include "PlasticSourceControlProvider.h"
 #include "PlasticSourceControlSettings.h"
 #include "PlasticSourceControlState.h"
 #include "HAL/PlatformFileManager.h"
@@ -15,6 +17,9 @@
 #include "Modules/ModuleManager.h"
 #include "XmlParser.h"
 #include "ISourceControlModule.h"
+#include "SoftwareVersion.h"
+#include "PlasticSourceControlChangelist.h"
+#include "PlasticSourceControlChangelistState.h"
 
 #if PLATFORM_LINUX
 #include <sys/ioctl.h>
@@ -35,6 +40,15 @@ namespace PlasticSourceControlConstants
 #endif
 }
 
+FScopedTempFile::FScopedTempFile(const FString& InText)
+{
+	Filename = FPaths::CreateTempFilename(*FPaths::ProjectLogDir(), TEXT("Plastic-Temp"), TEXT(".txt"));
+	if (!FFileHelper::SaveStringToFile(InText, *Filename, FFileHelper::EEncodingOptions::ForceUTF8WithoutBOM))
+	{
+		UE_LOG(LogSourceControl, Error, TEXT("Failed to write to temp file: %s"), *Filename);
+	}
+}
+
 FScopedTempFile::FScopedTempFile(const FText& InText)
 {
 	Filename = FPaths::CreateTempFilename(*FPaths::ProjectLogDir(), TEXT("Plastic-Temp"), TEXT(".txt"));
@@ -62,10 +76,6 @@ const FString& FScopedTempFile::GetFilename() const
 
 namespace PlasticSourceControlUtils
 {
-// Command-line interface parameters and output format changed with version 8.0.16.3000
-// For more details, see https://www.plasticscm.com/download/releasenotes/8.0.16.3000
-static bool				bIsNewVersion80163000 = false;
-
 // In/Out Pipes for the 'cm shell' persistent child process
 static void*			ShellOutputPipeRead = nullptr;
 static void*			ShellOutputPipeWrite = nullptr;
@@ -85,10 +95,12 @@ static void _CleanupBackgroundCommandLineShell()
 	ShellInputPipeRead = ShellInputPipeWrite = nullptr;
 }
 
-// Internal function to actualy launch the Plastic SCM background 'cm shell' process if possible (called under the critical section)
+// Internal function to launch the Plastic SCM background 'cm' process in interactive shell mode (called under the critical section)
 static bool _StartBackgroundPlasticShell(const FString& InPathToPlasticBinary, const FString& InWorkingDirectory)
 {
-	const FString FullCommand(TEXT("shell"));
+	TRACE_CPUPROFILER_EVENT_SCOPE(FPlasticSourceControlUtils::_StartBackgroundPlasticShell);
+
+	const FString FullCommand(TEXT("shell --encoding=UTF-8"));
 
 	const bool bLaunchDetached = false;				// the new process will NOT have its own window
 	const bool bLaunchHidden = true;				// the new process will be minimized in the task bar
@@ -119,6 +131,8 @@ static bool _StartBackgroundPlasticShell(const FString& InPathToPlasticBinary, c
 // Internal function (called under the critical section)
 static void _ExitBackgroundCommandLineShell()
 {
+	TRACE_CPUPROFILER_EVENT_SCOPE(FPlasticSourceControlUtils::_ExitBackgroundCommandLineShell);
+
 	if (ShellProcessHandle.IsValid())
 	{
 		if (FPlatformProcess::IsProcRunning(ShellProcessHandle))
@@ -132,7 +146,8 @@ static void _ExitBackgroundCommandLineShell()
 			{
 				if ((FPlatformTime::Seconds() - StartTimestamp) > Timeout)
 				{
-					UE_LOG(LogSourceControl, Warning, TEXT("ExitBackgroundCommandLineShell: cm shell didn't stop gracefuly in %lfs."), Timeout);
+					UE_LOG(LogSourceControl, Warning, TEXT("ExitBackgroundCommandLineShell: cm shell didn't stop gracefully in %lfs."), Timeout);
+					FPlatformProcess::TerminateProc(ShellProcessHandle);
 					break;
 				}
 				FPlatformProcess::Sleep(0.01f);
@@ -146,9 +161,9 @@ static void _ExitBackgroundCommandLineShell()
 // Internal function (called under the critical section)
 static void _RestartBackgroundCommandLineShell()
 {
-	const FPlasticSourceControlModule& PlasticSourceControl = FModuleManager::GetModuleChecked<FPlasticSourceControlModule>("PlasticSourceControl");
-	const FString& PathToPlasticBinary = PlasticSourceControl.AccessSettings().GetBinaryPath();
-	const FString& WorkingDirectory = PlasticSourceControl.GetProvider().GetPathToWorkspaceRoot();
+	const FPlasticSourceControlProvider& Provider = FPlasticSourceControlModule::Get().GetProvider();
+	const FString& PathToPlasticBinary = Provider.AccessSettings().GetBinaryPath();
+	const FString& WorkingDirectory = Provider.GetPathToWorkspaceRoot();
 
 	_ExitBackgroundCommandLineShell();
 	_StartBackgroundPlasticShell(PathToPlasticBinary, WorkingDirectory);
@@ -157,6 +172,8 @@ static void _RestartBackgroundCommandLineShell()
 // Internal function (called under the critical section)
 static bool _RunCommandInternal(const FString& InCommand, const TArray<FString>& InParameters, const TArray<FString>& InFiles, const EConcurrency::Type InConcurrency, FString& OutResults, FString& OutErrors)
 {
+	TRACE_CPUPROFILER_EVENT_SCOPE(FPlasticSourceControlUtils::_RunCommandInternal);
+
 	bool bResult = false;
 
 	ShellCommandCounter++;
@@ -164,11 +181,11 @@ static bool _RunCommandInternal(const FString& InCommand, const TArray<FString>&
 	// Detect previous crash of cm.exe and restart 'cm shell'
 	if (!FPlatformProcess::IsProcRunning(ShellProcessHandle))
 	{
-		UE_LOG(LogSourceControl, Warning, TEXT("RunCommand: 'cm shell' has stopped. Restarting! (count %d)"), ShellCommandCounter);
+		UE_LOG(LogSourceControl, Warning, TEXT("RunCommand: 'cm shell' has stopped. Restarting!"));
 		_RestartBackgroundCommandLineShell();
 	}
 
-	// Start with the Plastic command itself ("status", "log", "chekin"...)
+	// Start with the Plastic command itself ("status", "log", "checkin"...)
 	FString FullCommand = InCommand;
 	// Append to the command all parameters, and then finally the files
 	for (const FString& Parameter : InParameters)
@@ -186,21 +203,26 @@ static bool _RunCommandInternal(const FString& InCommand, const TArray<FString>&
 	UE_LOG(LogSourceControl, Verbose, TEXT("RunCommand: '%s' (%d chars, %d files)"), *LoggableCommand, FullCommand.Len()+1, InFiles.Num());
 	FullCommand += TEXT('\n'); // Finalize the command line
 
-	// Send command to 'cm shell' process
-	const bool bWriteOk = FPlatformProcess::WritePipe(ShellInputPipeWrite, FullCommand);
+	// Send command to 'cm shell' process in UTF-8
+	// NOTE: this explicit conversion to UTF-8 shouldn't be needed since FPlatformProcess::WritePipe() says it does it, but reading the implementation for Windows Platform show it merily truncates 16bits to 8bits chars!
+	// NOTE: on the other hand, ReadPipe() does the conversion from UTF-8 correctly already!
+	const FTCHARToUTF8 FullCommandUtf8(*FullCommand);
+	const bool bWriteOk = FPlatformProcess::WritePipe(ShellInputPipeWrite, reinterpret_cast<const uint8*>(FullCommandUtf8.Get()), FullCommandUtf8.Length());
 
-	// And wait up to 180.0 seconds for any kind of output from cm shell: in case of lengthier operation, intermediate output (like percentage of progress) is expected, which would refresh the timout
-	const double Timeout = 180.0;
+	// And wait up to 180.0 seconds for any kind of output from cm shell: in case of lengthier operation, intermediate output (like percentage of progress) is expected, which would refresh the timeout
+	static const double Timeout = 180.0;
 	const double StartTimestamp = FPlatformTime::Seconds();
 	double LastActivity = StartTimestamp;
 	double LastLog = StartTimestamp;
-	const double LogInterval = 5.0;
+	static const double LogInterval = 5.0;
 	int32 PreviousLogLen = 0;
 	while (FPlatformProcess::IsProcRunning(ShellProcessHandle))
 	{
 		FString Output = FPlatformProcess::ReadPipe(ShellOutputPipeRead);
 		if (!Output.IsEmpty())
 		{
+			TRACE_CPUPROFILER_EVENT_SCOPE(FPlasticSourceControlUtils::_RunCommandInternal::ParseOutput);
+
 			LastActivity = FPlatformTime::Seconds(); // freshen the timestamp while cm is still actively outputting information
 			OutResults.Append(MoveTemp(Output));
 			// Search the output for the line containing the result code, also indicating the end of the command
@@ -219,22 +241,27 @@ static bool _RunCommandInternal(const FString& InCommand, const TArray<FString>&
 				}
 			}
 		}
-		else if ((FPlatformTime::Seconds() - LastLog > LogInterval) && (PreviousLogLen < OutResults.Len()) && (InConcurrency == EConcurrency::Asynchronous))
+		else if ((FPlatformTime::Seconds() - LastLog > LogInterval) && (PreviousLogLen < OutResults.Len()))
 		{
 			// In case of long running operation, start to print intermediate output from cm shell (like percentage of progress)
-			// (but only when running Asynchronous commands, since Synchronous commands block the main thread until they finish)
 			UE_LOG(LogSourceControl, Log, TEXT("RunCommand: '%s' in progress for %.3lfs...\n%s"), *InCommand, (FPlatformTime::Seconds() - StartTimestamp), *OutResults.Mid(PreviousLogLen));
 			PreviousLogLen = OutResults.Len();
 			LastLog = FPlatformTime::Seconds(); // freshen the timestamp of last log
 		}
 		else if (FPlatformTime::Seconds() - LastActivity > Timeout)
 		{
-			// In case of timeout, ask the blocking 'cm shell' process to exit, detach from it and restart it immediatly
+			// In case of timeout, ask the blocking 'cm shell' process to exit, detach from it and restart it immediately
 			UE_LOG(LogSourceControl, Error, TEXT("RunCommand: '%s' TIMEOUT after %.3lfs output (%d chars):\n%s"), *InCommand, (FPlatformTime::Seconds() - StartTimestamp), OutResults.Len(), *OutResults.Mid(PreviousLogLen));
 			_RestartBackgroundCommandLineShell();
 			return false;
 		}
+		else if (IsEngineExitRequested())
+		{
+			UE_LOG(LogSourceControl, Warning, TEXT("RunCommand: '%s' Engine Exit was requested after %.3lfs output (%d chars):\n%s"), *InCommand, (FPlatformTime::Seconds() - StartTimestamp), OutResults.Len(), *OutResults.Mid(PreviousLogLen));
+			_ExitBackgroundCommandLineShell();
+		}
 
+		TRACE_CPUPROFILER_EVENT_SCOPE(FPlasticSourceControlUtils::_RunCommandInternal::Sleep);
 		FPlatformProcess::Sleep(0.001f);
 	}
 	const double ElapsedTime = (FPlatformTime::Seconds() - StartTimestamp);
@@ -326,6 +353,13 @@ bool RunCommand(const FString& InCommand, const TArray<FString>& InParameters, c
 	return bResult;
 }
 
+const FSoftwareVersion& GetOldestSupportedPlasticScmVersion()
+{
+	// https://www.plasticscm.com/download/releasenotes/9.0.16.4839 cm changelist 'persistent' flag now contain a '--' prefix.
+	static FSoftwareVersion s_OldestSupportedPlasticScmVersion(TEXT("9.0.16.4839"));
+	return s_OldestSupportedPlasticScmVersion;
+}
+
 FString FindPlasticBinaryPath()
 {
 #if PLATFORM_WINDOWS
@@ -380,49 +414,18 @@ bool FindRootDirectory(const FString& InPath, FString& OutWorkspaceRoot)
 	return bFound;
 }
 
-/**
- * @brief Compare Plastic SCM cli version strings.
- * @param VersionA		PlasticSCM version string in the form "0.0.0.0" (as returned by GetPlasticScmVersion)
- * @param VersionB		PlasticSCM version string in the form "0.0.0.0" (as returned by GetPlasticScmVersion)
- * @returns true if VersionA is lower than VersionB
-*/
-static bool PlasticScmVersionLess(const FString& VersionA, const FString& VersionB) {
-	struct Version {
-		Version(const FString& V) {
-			TArray<FString> Parts;
-			const int32 N = V.ParseIntoArray(Parts, TEXT("."));
-			check(N == 4);
-			a = FCString::Atoi(*Parts[0]);
-			b = FCString::Atoi(*Parts[1]);
-			c = FCString::Atoi(*Parts[2]);
-			d = FCString::Atoi(*Parts[3]);
-		}
-		int a, b, c, d;
-	} A(VersionA), B(VersionB);
-	if (A.a < B.a) return true;
-	if (B.a < A.a) return false;
-	if (A.b < B.b) return true;
-	if (B.b < A.b) return false;
-	if (A.c < B.c) return true;
-	if (B.c < A.c) return false;
-	if (A.d < B.d) return true;
-	if (B.d < A.d) return false;
-	return false; // Equal
-}
-
 // This is called once by FPlasticSourceControlProvider::CheckPlasticAvailability()
-void GetPlasticScmVersion(FString& OutPlasticScmVersion)
+bool GetPlasticScmVersion(FSoftwareVersion& OutPlasticScmVersion)
 {
 	TArray<FString> InfoMessages;
 	TArray<FString> ErrorMessages;
 	const bool bResult = RunCommand(TEXT("version"), TArray<FString>(), TArray<FString>(), EConcurrency::Synchronous, InfoMessages, ErrorMessages);
 	if (bResult && InfoMessages.Num() > 0)
 	{
-		OutPlasticScmVersion = InfoMessages[0];
-
-		// Command-line format output changed with version 8.0.16.3000
-		bIsNewVersion80163000 = !PlasticScmVersionLess(OutPlasticScmVersion, "8.0.16.3000");
+		OutPlasticScmVersion = FSoftwareVersion(InfoMessages[0]);
+		return true;
 	}
+	return false;
 }
 
 void GetUserName(FString& OutUserName)
@@ -436,19 +439,20 @@ void GetUserName(FString& OutUserName)
 	}
 }
 
-bool GetWorkspaceName(FString& OutWorkspaceName)
+bool GetWorkspaceName(const FString& InWorkspaceRoot, FString& OutWorkspaceName, TArray<FString>& OutErrorMessages)
 {
 	TArray<FString> InfoMessages;
-	TArray<FString> ErrorMessages;
+
 	TArray<FString> Parameters;
-	Parameters.Add(TEXT("."));
 	Parameters.Add(TEXT("--format={0}"));
+	TArray<FString> Files;
+	Files.Add(InWorkspaceRoot);
 	// Get the workspace name
-	const bool bResult = RunCommand(TEXT("getworkspacefrompath"), Parameters, TArray<FString>(), EConcurrency::Synchronous, InfoMessages, ErrorMessages);
+	const bool bResult = RunCommand(TEXT("getworkspacefrompath"), Parameters, Files, EConcurrency::Synchronous, InfoMessages, OutErrorMessages);
 	if (bResult && InfoMessages.Num() > 0)
 	{
-		// NOTE: getworkspacefrompath never returns an error!
-		if (!InfoMessages[0].Equals(TEXT(". is not in a workspace.")))
+		// NOTE: an old version of cm getworkspacefrompath didn't return an error code so we had to rely on the error message
+		if (!InfoMessages[0].EndsWith(TEXT(" is not in a workspace.")))
 		{
 			OutWorkspaceName = MoveTemp(InfoMessages[0]);
 		}
@@ -498,25 +502,15 @@ static bool ParseWorkspaceInformation(const TArray<FString>& InInfoMessages, int
 	return bResult;
 }
 
-bool GetWorkspaceInformation(int32& OutChangeset, FString& OutRepositoryName, FString& OutServerUrl, FString& OutBranchName)
+bool GetWorkspaceInformation(int32& OutChangeset, FString& OutRepositoryName, FString& OutServerUrl, FString& OutBranchName, TArray<FString>& OutErrorMessages)
 {
 	TArray<FString> InfoMessages;
-	TArray<FString> ErrorMessages;
 	TArray<FString> Parameters;
-
-	// Command-line format output changed with version 8.0.16.3000, see https://www.plasticscm.com/download/releasenotes/8.0.16.3000
-	if (bIsNewVersion80163000)
-	{
-		Parameters.Add(TEXT("--compact"));
-		Parameters.Add(TEXT("--header")); // Only prints the workspace status. No file status.
-	}
-	else
-	{
-		Parameters.Add(TEXT("--nochanges")); // Only prints the workspace status. No file status.
-	}
+	Parameters.Add(TEXT("--compact"));
+	Parameters.Add(TEXT("--header")); // Only prints the workspace status. No file status.
 	// NOTE: --wkconfig results in two network calls GetBranchInfoByName & GetLastChangesetOnBranch so it's okay to do it once here but not all the time
 	Parameters.Add(TEXT("--wkconfig")); // Branch name. NOTE: Deprecated in 8.0.16.3000 https://www.plasticscm.com/download/releasenotes/8.0.16.3000
-	bool bResult = RunCommand(TEXT("status"), Parameters, TArray<FString>(), EConcurrency::Synchronous, InfoMessages, ErrorMessages);
+	bool bResult = RunCommand(TEXT("status"), Parameters, TArray<FString>(), EConcurrency::Synchronous, InfoMessages, OutErrorMessages);
 	if (bResult)
 	{
 		bResult = ParseWorkspaceInformation(InfoMessages, OutChangeset, OutRepositoryName, OutServerUrl, OutBranchName);
@@ -525,6 +519,24 @@ bool GetWorkspaceInformation(int32& OutChangeset, FString& OutRepositoryName, FS
 	return bResult;
 }
 
+FString UserNameToDisplayName(const FString& InUserName)
+{
+	if (const FString* Result = GetDefault<UPlasticSourceControlProjectSettings>()->UserNameToDisplayName.Find(InUserName))
+	{
+		return *Result;
+	}
+	else if (GetDefault<UPlasticSourceControlProjectSettings>()->bHideEmailDomainInUsername)
+	{
+		int32 EmailDomainSeparatorIndex;
+		if (InUserName.FindChar(TEXT('@'), EmailDomainSeparatorIndex))
+		{
+			return InUserName.Left(EmailDomainSeparatorIndex);
+		}
+	}
+
+	return InUserName;
+}
+
 /**
  * @brief Extract the renamed from filename from a Plastic SCM status result.
  *
@@ -588,7 +600,7 @@ static FString FilenameFromPlasticStatus(const FString& InResult)
 class FPlasticStatusFileMatcher
 {
 public:
-	FPlasticStatusFileMatcher(const FString& InAbsoluteFilename)
+	explicit FPlasticStatusFileMatcher(const FString& InAbsoluteFilename)
 		: AbsoluteFilename(InAbsoluteFilename)
 	{
 	}
@@ -620,7 +632,7 @@ private:
 static EWorkspaceState::Type StateFromPlasticStatus(const FString& InResult)
 {
 	EWorkspaceState::Type State;
-	const FString FileStatus = InResult.Mid(1, 2);
+	const FString FileStatus = (InResult[0] == TEXT(' ')) ? InResult.Mid(1, 2) : InResult;
 
 	if (FileStatus == "CH") // Modified but not Checked-Out
 	{
@@ -695,8 +707,9 @@ static EWorkspaceState::Type StateFromPlasticStatus(const FString& InResult)
  */
 static void ParseFileStatusResult(TArray<FString>&& InFiles, const TArray<FString>& InResults, TArray<FPlasticSourceControlState>& OutStates, int32& OutChangeset, FString& OutBranchName)
 {
-	const FPlasticSourceControlModule& PlasticSourceControl = FModuleManager::GetModuleChecked<FPlasticSourceControlModule>("PlasticSourceControl");
-	const FString& WorkingDirectory = PlasticSourceControl.GetProvider().GetPathToWorkspaceRoot();
+	TRACE_CPUPROFILER_EVENT_SCOPE(FPlasticSourceControlUtils::ParseFileStatusResult);
+
+	const FString& WorkingDirectory = FPlasticSourceControlModule::Get().GetProvider().GetPathToWorkspaceRoot();
 
 	// Parse the first two lines with Changeset number and Branch name (the second being requested only once at init)
 	FString RepositoryName, ServerUrl;
@@ -766,8 +779,9 @@ static void ParseFileStatusResult(TArray<FString>&& InFiles, const TArray<FStrin
 */
 static void ParseDirectoryStatusResultForDeleted(const TArray<FString>& InResults, TArray<FPlasticSourceControlState>& OutStates)
 {
-	const FPlasticSourceControlModule& PlasticSourceControl = FModuleManager::GetModuleChecked<FPlasticSourceControlModule>("PlasticSourceControl");
-	const FString& WorkingDirectory = PlasticSourceControl.GetProvider().GetPathToWorkspaceRoot();
+	TRACE_CPUPROFILER_EVENT_SCOPE(FPlasticSourceControlUtils::ParseDirectoryStatusResultForDeleted);
+
+	const FString& WorkingDirectory = FPlasticSourceControlModule::Get().GetProvider().GetPathToWorkspaceRoot();
 
 	// Iterate on each line of result of the status command
 	for (const FString& Result : InResults)
@@ -821,15 +835,12 @@ public:
  */
 static bool RunStatus(const FString& InDir, TArray<FString>&& InFiles, const EConcurrency::Type InConcurrency, TArray<FString>& OutErrorMessages, TArray<FPlasticSourceControlState>& OutStates, int32& OutChangeset, FString& OutBranchName)
 {
+	TRACE_CPUPROFILER_EVENT_SCOPE(FPlasticSourceControlUtils::RunStatus);
+
 	check(InFiles.Num() > 0);
 
 	TArray<FString> Parameters;
-
-	// Command-line format output changed with version 8.0.16.3000, see https://www.plasticscm.com/download/releasenotes/8.0.16.3000
-	if (bIsNewVersion80163000)
-	{
-		Parameters.Add(TEXT("--compact"));
-	}
+	Parameters.Add(TEXT("--compact"));
 	Parameters.Add(TEXT("--noheaders"));
 	Parameters.Add(TEXT("--all"));
 	Parameters.Add(TEXT("--ignored"));
@@ -872,7 +883,7 @@ static bool RunStatus(const FString& InDir, TArray<FString>&& InFiles, const ECo
 			// => so we also parse the status results to explicitly look for Removed/Deleted assets
 			if (Results.Num() > 0)
 			{
-				Results.RemoveAt(0, 1);// Before that, remove the first line (Workspace/Changeset info)
+				Results.RemoveAt(0, 1); // Before that, remove the first line (Workspace/Changeset info)
 			}
 			ParseDirectoryStatusResultForDeleted(Results, OutStates);
 		}
@@ -892,26 +903,17 @@ static bool RunStatus(const FString& InDir, TArray<FString>&& InFiles, const ECo
 class FPlasticFileinfoParser
 {
 public:
-	FPlasticFileinfoParser(const FString& InResult)
+	explicit FPlasticFileinfoParser(const FString& InResult)
 	{
 		TArray<FString> Fileinfos;
-		const int32 NbElmts = InResult.ParseIntoArray(Fileinfos, TEXT(";"));
-		if (NbElmts >= 2)
+		const int32 NbElmts = InResult.ParseIntoArray(Fileinfos, TEXT(";"), false); // Don't cull empty values in csv
+		if (NbElmts == 5)
 		{
 			RevisionChangeset = FCString::Atoi(*Fileinfos[0]);
 			RevisionHeadChangeset = FCString::Atoi(*Fileinfos[1]);
-			if (NbElmts >= 3)
-			{
-				RepSpec = MoveTemp(Fileinfos[2]);
-				if (NbElmts >=4)
-				{
-					LockedBy = MoveTemp(Fileinfos[3]);
-					if (NbElmts >= 5)
-					{
-						LockedWhere = MoveTemp(Fileinfos[4]);
-					}
-				}
-			}
+			RepSpec = MoveTemp(Fileinfos[2]);
+			LockedBy = UserNameToDisplayName(MoveTemp(Fileinfos[3]));
+			LockedWhere = MoveTemp(Fileinfos[4]);
 		}
 	}
 
@@ -931,8 +933,9 @@ public:
  */
 static void ParseFileinfoResults(const TArray<FString>& InResults, TArray<FPlasticSourceControlState>& InOutStates)
 {
-	const FPlasticSourceControlModule& PlasticSourceControl = FModuleManager::GetModuleChecked<FPlasticSourceControlModule>("PlasticSourceControl");
-	const FPlasticSourceControlProvider& Provider = PlasticSourceControl.GetProvider();
+	TRACE_CPUPROFILER_EVENT_SCOPE(FPlasticSourceControlUtils::ParseFileinfoResults);
+
+	const FPlasticSourceControlProvider& Provider = FPlasticSourceControlModule::Get().GetProvider();
 
 	ensureMsgf(InResults.Num() == InOutStates.Num(), TEXT("The fileinfo command should gives the same number of infos as the status command"));
 
@@ -951,7 +954,7 @@ static void ParseFileinfoResults(const TArray<FString>& InResults, TArray<FPlast
 		FileState.LockedWhere = MoveTemp(FileinfoParser.LockedWhere);
 
 		// If a file is locked but not checked-out locally (or moved/renamed) this means it is locked by someone else or elsewhere
-		if ((FileState.WorkspaceState != EWorkspaceState::CheckedOut) && (FileState.WorkspaceState != EWorkspaceState::Moved) && !FileState.LockedBy.IsEmpty()) 
+		if ((FileState.WorkspaceState != EWorkspaceState::CheckedOut) && (FileState.WorkspaceState != EWorkspaceState::Moved) && !FileState.LockedBy.IsEmpty())
 		{
 			UE_LOG(LogSourceControl, Verbose, TEXT("LockedByOther(%s) by '%s!=%s' (or %s!=%s)"), *File, *FileState.LockedBy, *Provider.GetUserName(), *FileState.LockedWhere, *Provider.GetWorkspaceName());
 			FileState.WorkspaceState = EWorkspaceState::LockedByOther;
@@ -975,7 +978,7 @@ static void ParseFileinfoResults(const TArray<FString>& InResults, TArray<FPlast
  *
  * ie RevisionChangeset, RevisionHeadChangeset, RepSpec, LockedBy, LockedWhere
  *
- * @param[in]		bInWholeDirectory	If executed on a whole directory (typically Content/) for a "Submit Content" operation, optimize fileinfo more agressively
+ * @param[in]		bInWholeDirectory	If executed on a whole directory (typically Content/) for a "Submit Content" operation, optimize fileinfo more aggressively
  * @param			bInUpdateHistory	If getting the history of files, force execute the fileinfo command required to get RepSpec of XLinks (history view or visual diff)
  * @param[in]		InConcurrency		Is the command running in the background, or blocking the main thread
  * @param[out]		OutErrorMessages	Error messages from the "fileinfo" command
@@ -983,6 +986,8 @@ static void ParseFileinfoResults(const TArray<FString>& InResults, TArray<FPlast
  */
 static bool RunFileinfo(const bool bInWholeDirectory, const bool bInUpdateHistory, const EConcurrency::Type InConcurrency, TArray<FString>& OutErrorMessages, TArray<FPlasticSourceControlState>& InOutStates)
 {
+	TRACE_CPUPROFILER_EVENT_SCOPE(FPlasticSourceControlUtils::RunFileinfo);
+
 	bool bResult = true;
 	TArray<FString> SelectedFiles;
 
@@ -990,16 +995,16 @@ static bool RunFileinfo(const bool bInWholeDirectory, const bool bInUpdateHistor
 	TArray<FPlasticSourceControlState> OptimizedStates;
 	for (FPlasticSourceControlState& State : InOutStates)
 	{
-		// 1) Issue a "fileinfo" command for controled files (to know if they are up to date and can be checked-out or checked-in)
+		// 1) Issue a "fileinfo" command for controlled files (to know if they are up to date and can be checked-out or checked-in)
 		// but only if controlled unchanged, or locally changed / locally deleted,
 		// optimizing for files that are CheckedOut/Added/Deleted/Moved/Copied/Replaced/NotControled/Ignored/Private/Unknown
-		// (since there is no point to check if they are up to date in these cases; they are already checkedout or not controlld).
+		// (since there is no point to check if they are up to date in these cases; they are already checked-out or not controlled).
 		// This greatly reduce the time needed to do some operations like "Add" or "Move/Rename/Copy" when there is some latency with the server (eg cloud).
 		//
 		// 2) bInWholeDirectory: In the case of a "whole directory status" triggered by the "Submit Content" operation,
 		// don't even issue a "fileinfo" command for unchanged Controlled files since they won't be considered them for submit.
 		// This greatly reduce the time needed to open the Submit window.
-		// 
+		//
 		// 3) bInUpdateHistory: When the plugin needs to update the history of files, it needs to know if it's on a XLink,
 		// so the fileinfo command is required here to get the RepSpec
 		if (bInUpdateHistory
@@ -1041,7 +1046,7 @@ static bool RunFileinfo(const bool bInWholeDirectory, const bool bInUpdateHistor
 class FPlasticMergeConflictParser
 {
 public:
-	FPlasticMergeConflictParser(const FString& InResult)
+	explicit FPlasticMergeConflictParser(const FString& InResult)
 	{
 		static const FString FILE_CONFLICT(TEXT("FILE_CONFLICT "));
 		if (InResult.StartsWith(FILE_CONFLICT, ESearchCase::CaseSensitive))
@@ -1073,11 +1078,12 @@ public:
 };
 
 // Check if merging, and from which changelist, then execute a cm merge command to amend status for listed files
-bool RunCheckMergeStatus(const TArray<FString>& InFiles, TArray<FString>& OutErrorMessages, TArray<FPlasticSourceControlState>& OutStates)
+static bool RunCheckMergeStatus(const TArray<FString>& InFiles, TArray<FString>& OutErrorMessages, TArray<FPlasticSourceControlState>& OutStates)
 {
+	TRACE_CPUPROFILER_EVENT_SCOPE(FPlasticSourceControlUtils::RunCheckMergeStatus);
+
 	bool bResult = false;
-	FPlasticSourceControlModule& PlasticSourceControl = FModuleManager::GetModuleChecked<FPlasticSourceControlModule>("PlasticSourceControl");
-	FPlasticSourceControlProvider& Provider = PlasticSourceControl.GetProvider();
+	const FPlasticSourceControlProvider& Provider = FPlasticSourceControlModule::Get().GetProvider();
 
 	const FString MergeProgressFilename = FPaths::Combine(*Provider.GetPathToWorkspaceRoot(), TEXT(".plastic/plastic.mergeprogress"));
 	if (FPaths::FileExists(MergeProgressFilename))
@@ -1130,7 +1136,7 @@ bool RunCheckMergeStatus(const TArray<FString>& InFiles, TArray<FString>& OutErr
 					// Store the Merge Parameters for reuse with later "Resolve" operation
 					const TArray<FString> PendingMergeParameters = Parameters;
 					Parameters.Add(TEXT("--machinereadable"));
-					// call 'cm merge cs:xxx --machinereadable' (only dry-run, whithout the --merge parameter)
+					// call 'cm merge cs:xxx --machinereadable' (only dry-run, without the --merge parameter)
 					bResult = RunCommand(TEXT("merge"), Parameters, TArray<FString>(), EConcurrency::Synchronous, Results, ErrorMessages);
 					OutErrorMessages.Append(MoveTemp(ErrorMessages));
 					// Parse the result, one line for each conflicted files:
@@ -1168,9 +1174,13 @@ FString FindCommonDirectory(const FString& InPath1, const FString& InPath2)
 	for (int32 Index = 0; Index < MinLen; Index++)
 	{
 		if (InPath1[Index] != InPath2[Index])
+		{
 			break;
+		}
 		if (InPath1[Index] == TEXT('/'))
+		{
 			IndexAfterLastCommonSeparator = Index + 1;
+		}
 	}
 	return InPath1.Left(IndexAfterLastCommonSeparator);
 }
@@ -1188,15 +1198,14 @@ bool RunUpdateStatus(const TArray<FString>& InFiles, const bool bInUpdateHistory
 {
 	bool bResults = true;
 
-	const FPlasticSourceControlModule& PlasticSourceControl = FModuleManager::GetModuleChecked<FPlasticSourceControlModule>("PlasticSourceControl");
-	const FString& WorkspaceRoot = PlasticSourceControl.GetProvider().GetPathToWorkspaceRoot();
+	const FString& WorkspaceRoot = FPlasticSourceControlModule::Get().GetProvider().GetPathToWorkspaceRoot();
 
 	// The "status" command only operate on one directory-tree at a time (whole tree recursively)
 	// not on different folders with no common root.
 	// But "Submit to Source Control" ask for the State of many different directories,
 	// from Project/Content and Project/Config, Engine/Content, Engine/Plugins/<...>/Content...
 
-	// In a similar way, a check-in can involve files from different subdirectories, and UpdateStatus is called for all of them at once.
+	// In a similar way, a checkin can involve files from different subdirectories, and UpdateStatus is called for all of them at once.
 
 	static TArray<FString> RootDirs =
 	{
@@ -1275,7 +1284,7 @@ bool RunUpdateStatus(const TArray<FString>& InFiles, const bool bInUpdateHistory
 	// 2) then we can batch Plastic status operation by subdirectory
 	for (auto& Group : GroupOfFiles)
 	{
-		const bool bWholeDirectory = ((Group.Value.Files.Num() == 1) && (Group.Value.CommonDir == Group.Value.Files[0]));	
+		const bool bWholeDirectory = ((Group.Value.Files.Num() == 1) && (Group.Value.CommonDir == Group.Value.Files[0]));
 
 		// Run a "status" command on the directory to get workspace file states.
 		// (ie. Changed, CheckedOut, Copied, Replaced, Added, Private, Ignored, Deleted, LocallyDeleted, Moved, LocallyMoved)
@@ -1285,7 +1294,7 @@ bool RunUpdateStatus(const TArray<FString>& InFiles, const bool bInUpdateHistory
 		{
 			// Run a "fileinfo" command to update complementary status information of given files.
 			// (ie RevisionChangeset, RevisionHeadChangeset, RepSpec, LockedBy, LockedWhere)
-			// In case of "whole directory status", there is no explicit file in the group (it contains only the directory) 
+			// In case of "whole directory status", there is no explicit file in the group (it contains only the directory)
 			// => work on the list of files discovered by RunStatus()
 			bResults &= RunFileinfo(bWholeDirectory, bInUpdateHistory, InConcurrency, OutErrorMessages, States);
 		}
@@ -1302,6 +1311,8 @@ bool RunUpdateStatus(const TArray<FString>& InFiles, const bool bInUpdateHistory
 // cm cat revid:1230@rep:myrep@repserver:myserver:8084 --raw --file=Name124.tmp
 bool RunDumpToFile(const FString& InPathToPlasticBinary, const FString& InRevSpec, const FString& InDumpFileName)
 {
+	TRACE_CPUPROFILER_EVENT_SCOPE(FPlasticSourceControlUtils::RunDumpToFile);
+
 	int32	ReturnCode = 0;
 	FString Results;
 	FString Errors;
@@ -1324,6 +1335,17 @@ bool RunDumpToFile(const FString& InPathToPlasticBinary, const FString& InRevSpe
 	return bResult;
 }
 
+FString DecodeXmlEntities(const FString& InString)
+{
+	FString String = InString;
+	String.ReplaceInline(TEXT("&amp;"), TEXT("&"));
+	String.ReplaceInline(TEXT("&quot;"), TEXT("\""));
+	String.ReplaceInline(TEXT("&apos;"), TEXT("'"));
+	String.ReplaceInline(TEXT("&lt;"), TEXT("<"));
+	String.ReplaceInline(TEXT("&gt;"), TEXT(">"));
+	return String;
+}
+
 /**
  * Parse results of the 'cm history --moveddeleted --xml --encoding="utf-8"' command.
  * 
@@ -1370,8 +1392,9 @@ bool RunDumpToFile(const FString& InPathToPlasticBinary, const FString& InRevSpe
 */
 static bool ParseHistoryResults(const bool bInUpdateHistory, const FXmlFile& InXmlResult, TArray<FPlasticSourceControlState>& InOutStates)
 {
-	const FPlasticSourceControlModule& PlasticSourceControl = FModuleManager::GetModuleChecked<FPlasticSourceControlModule>("PlasticSourceControl");
-	const FPlasticSourceControlProvider& Provider = PlasticSourceControl.GetProvider();
+	TRACE_CPUPROFILER_EVENT_SCOPE(FPlasticSourceControlUtils::ParseHistoryResults);
+
+	const FPlasticSourceControlProvider& Provider = FPlasticSourceControlModule::Get().GetProvider();
 	const FString RootRepSpec = FString::Printf(TEXT("%s@%s"), *Provider.GetRepositoryName(), *Provider.GetServerUrl());
 
 	static const FString RevisionHistoriesResult(TEXT("RevisionHistoriesResult"));
@@ -1481,11 +1504,11 @@ static bool ParseHistoryResults(const bool bInUpdateHistory, const FXmlFile& InX
 				}
 				if (const FXmlNode* CommentNode = RevisionNode->FindChildNode(Comment))
 				{
-					SourceControlRevision->Description = CommentNode->GetContent();
+					SourceControlRevision->Description = DecodeXmlEntities(CommentNode->GetContent());
 				}
 				if (const FXmlNode* OwnerNode = RevisionNode->FindChildNode(Owner))
 				{
-					SourceControlRevision->UserName = OwnerNode->GetContent();
+					SourceControlRevision->UserName = UserNameToDisplayName(OwnerNode->GetContent());
 				}
 				if (const FXmlNode* DateNode = RevisionNode->FindChildNode(CreationDate))
 				{
@@ -1518,6 +1541,12 @@ static bool ParseHistoryResults(const bool bInUpdateHistory, const FXmlFile& InX
 					InOutState.History.Add(SourceControlRevision);
 				}
 
+				// Also grab the UserName of the author of the current depot/head changeset
+				if ((SourceControlRevision->ChangesetNumber == InOutState.DepotRevisionChangeset) && InOutState.HeadUserName.IsEmpty())
+				{
+					InOutState.HeadUserName = SourceControlRevision->UserName;
+				}
+
 				if (!bInUpdateHistory)
 				{
 					break; // if not updating the history, just getting the head of the latest branch is enough
@@ -1563,7 +1592,10 @@ bool RunGetHistory(const bool bInUpdateHistory, TArray<FPlasticSourceControlStat
 		if (bResult)
 		{
 			FXmlFile XmlFile;
-			bResult = XmlFile.LoadFile(Results, EConstructMethod::ConstructFromBuffer);
+			{
+				TRACE_CPUPROFILER_EVENT_SCOPE(FPlasticSourceControlUtils::RunGetHistory::FXmlFile::LoadFile);
+				bResult = XmlFile.LoadFile(Results, EConstructMethod::ConstructFromBuffer);
+			}
 			if (bResult)
 			{
 				bResult = ParseHistoryResults(bInUpdateHistory, XmlFile, InOutStates);
@@ -1578,10 +1610,171 @@ bool RunGetHistory(const bool bInUpdateHistory, TArray<FPlasticSourceControlStat
 	return bResult;
 }
 
+
+/**
+ * Parse results of the 'cm status --changelists --controlledchanged --xml --encoding="utf-8"' command.
+ *
+ * Results of the status changelists command looks like that:
+<StatusOutput>
+  <WorkspaceStatus>
+	<Status>
+	  <RepSpec>
+		<Server>test@cloud</Server>
+		<Name>UEPlasticPluginDev</Name>
+	  </RepSpec>
+	  <Changeset>11</Changeset>
+	</Status>
+  </WorkspaceStatus>
+  <WkConfigType>Branch</WkConfigType>
+  <WkConfigName>/main@rep:UEPlasticPluginDev@repserver:test@cloud</WkConfigName>
+  <Changelists>
+	<Changelist>
+	  <Name>Default</Name>
+	  <Description>Default Plastic SCM changelist</Description>
+	  <Changes>
+		<Change>
+		  <Type>CO</Type>
+		  <TypeVerbose>Checked-out</TypeVerbose>
+		  <Path>UEPlasticPluginDev.uproject</Path>
+		  <OldPath />
+		  <PrintableMovedPath />
+		  <MergesInfo />
+		  <SimilarityPerUnit>0</SimilarityPerUnit>
+		  <Similarity />
+		  <Size>583</Size>
+		  <PrintableSize>583 bytes</PrintableSize>
+		  <PrintableLastModified>6 days ago</PrintableLastModified>
+		  <RevisionType>enTextFile</RevisionType>
+		  <LastModified>2022-06-07T12:28:32+02:00</LastModified>
+		</Change>
+		[...]
+		<Change>
+		</Change>
+	  </Changes>
+	</Changelist>
+  </Changelists>
+</StatusOutput>
+*/
+static bool ParseChangelistsResults(const FXmlFile& InXmlResult, TArray<FPlasticSourceControlChangelistState>& OutChangelistsStates, TArray<TArray<FPlasticSourceControlState>>& OutCLFilesStates)
+{
+	TRACE_CPUPROFILER_EVENT_SCOPE(FPlasticSourceControlUtils::ParseChangelistsResults);
+
+	static const FString StatusOutput(TEXT("StatusOutput"));
+	static const FString WkConfigType(TEXT("WkConfigType"));
+	static const FString WkConfigName(TEXT("WkConfigName"));
+	static const FString Changelists(TEXT("Changelists"));
+	static const FString Changelist(TEXT("Changelist"));
+	static const FString Name(TEXT("Name"));
+	static const FString Description(TEXT("Description"));
+	static const FString Changes(TEXT("Changes"));
+	static const FString Change(TEXT("Change"));
+	static const FString Type(TEXT("Type"));
+	static const FString Path(TEXT("Path"));
+
+	const FString& WorkingDirectory = FPlasticSourceControlModule::Get().GetProvider().GetPathToWorkspaceRoot();
+
+	const FXmlNode* StatusOutputNode = InXmlResult.GetRootNode();
+	if (StatusOutputNode == nullptr || StatusOutputNode->GetTag() != StatusOutput)
+	{
+		return false;
+	}
+
+	const FXmlNode* ChangelistsNode = StatusOutputNode->FindChildNode(Changelists);
+	if (ChangelistsNode)
+	{
+		const TArray<FXmlNode*>& ChangelistNodes = ChangelistsNode->GetChildrenNodes();
+		OutCLFilesStates.SetNum(ChangelistNodes.Num());
+		for (int32 ChangelistIndex = 0; ChangelistIndex < ChangelistNodes.Num(); ChangelistIndex++)
+		{
+			const FXmlNode* ChangelistNode = ChangelistNodes[ChangelistIndex];
+			check(ChangelistNode);
+			const FXmlNode* NameNode = ChangelistNode->FindChildNode(Name);
+			const FXmlNode* DescriptionNode = ChangelistNode->FindChildNode(Description);
+			const FXmlNode* ChangesNode = ChangelistNode->FindChildNode(Changes);
+			if (NameNode == nullptr || DescriptionNode == nullptr || ChangesNode == nullptr)
+			{
+				continue;
+			}
+
+			FString NameTemp = PlasticSourceControlUtils::DecodeXmlEntities(NameNode->GetContent());
+			FPlasticSourceControlChangelist ChangelistTemp(MoveTemp(NameTemp), true);
+			FString DescriptionTemp = ChangelistTemp.IsDefault() ? FString() : PlasticSourceControlUtils::DecodeXmlEntities(DescriptionNode->GetContent());
+			FPlasticSourceControlChangelistState ChangelistState(MoveTemp(ChangelistTemp), MoveTemp(DescriptionTemp));
+
+			const TArray<FXmlNode*>& ChangeNodes = ChangesNode->GetChildrenNodes();
+			for (const FXmlNode* ChangeNode : ChangeNodes)
+			{
+				check(ChangeNode);
+				const FXmlNode* PathNode = ChangeNode->FindChildNode(Path);
+				const FXmlNode* TypeNode = ChangeNode->FindChildNode(Type);
+				if (PathNode == nullptr || TypeNode == nullptr)
+				{
+					continue;
+				}
+
+				// Here we make sure to only collect file states, since we shouldn't display the added directories to the Editor
+				FString FileName = PathNode->GetContent();
+				int32 DotIndex;
+				if (FileName.FindChar(TEXT('.'), DotIndex))
+				{
+					FPlasticSourceControlState FileState(FPaths::ConvertRelativePathToFull(WorkingDirectory, MoveTemp(FileName)));
+					FileState.WorkspaceState = PlasticSourceControlUtils::StateFromPlasticStatus(TypeNode->GetContent());
+					FileState.Changelist = ChangelistState.Changelist;
+					OutCLFilesStates[ChangelistIndex].Add(MoveTemp(FileState));
+				}
+			}
+
+			OutChangelistsStates.Add(ChangelistState);
+		}
+	}
+
+	if (!OutChangelistsStates.FindByPredicate(
+		[](const FPlasticSourceControlChangelistState& CLState) { return CLState.Changelist.IsDefault(); }
+	))
+	{
+		// No Default Changelists isn't an error, but the Editor UX expects to always the Default changelist (so you can always move files back to it)
+		FPlasticSourceControlChangelistState DefaultChangelistState(FPlasticSourceControlChangelist::DefaultChangelist);
+		OutChangelistsStates.Insert(DefaultChangelistState, 0);
+		OutCLFilesStates.Insert(TArray<FPlasticSourceControlState>(), 0);
+	}
+
+	return true;
+}
+
+bool RunGetChangelists(const EConcurrency::Type InConcurrency, TArray<FPlasticSourceControlChangelistState>& OutChangelistsStates, TArray<TArray<FPlasticSourceControlState>>& OutCLFilesStates, TArray<FString>& OutErrorMessages)
+{
+	FString Results;
+	FString Errors;
+	TArray<FString> Parameters;
+	Parameters.Add(TEXT("--changelists"));
+	Parameters.Add(TEXT("--controlledchanged"));
+	Parameters.Add(TEXT("--noheader"));
+	Parameters.Add(TEXT("--xml"));
+	Parameters.Add(TEXT("--encoding=\"utf-8\""));
+	bool bResult = PlasticSourceControlUtils::RunCommandInternal(TEXT("status"), Parameters, TArray<FString>(), InConcurrency, Results, Errors);
+	if (bResult)
+	{
+		FXmlFile XmlFile;
+		{
+			TRACE_CPUPROFILER_EVENT_SCOPE(FPlasticSourceControlUtils::RunGetHistory::FXmlFile::LoadFile);
+			bResult = XmlFile.LoadFile(Results, EConstructMethod::ConstructFromBuffer);
+		}
+		if (bResult)
+		{
+			bResult = ParseChangelistsResults(XmlFile, OutChangelistsStates, OutCLFilesStates);
+		}
+	}
+	if (!Errors.IsEmpty())
+	{
+		OutErrorMessages.Add(MoveTemp(Errors));
+	}
+
+	return bResult;
+}
+
 bool UpdateCachedStates(TArray<FPlasticSourceControlState>&& InStates)
 {
-	FPlasticSourceControlModule& PlasticSourceControl = FModuleManager::GetModuleChecked<FPlasticSourceControlModule>( "PlasticSourceControl" );
-	FPlasticSourceControlProvider& Provider = PlasticSourceControl.GetProvider();
+	FPlasticSourceControlProvider& Provider = FPlasticSourceControlModule::Get().GetProvider();
 	const FDateTime Now = FDateTime::Now();
 
 	for (auto&& InState : InStates)
@@ -1599,7 +1792,7 @@ bool UpdateCachedStates(TArray<FPlasticSourceControlState>&& InStates)
  */
 struct FRemoveRedundantErrors
 {
-	FRemoveRedundantErrors(const FString& InFilter)
+	explicit FRemoveRedundantErrors(const FString& InFilter)
 		: Filter(InFilter)
 	{
 	}
@@ -1630,7 +1823,7 @@ void RemoveRedundantErrors(FPlasticSourceControlCommand& InCommand, const FStrin
 		}
 	}
 
-	InCommand.ErrorMessages.RemoveAll( FRemoveRedundantErrors(InFilter) );
+	InCommand.ErrorMessages.RemoveAll(FRemoveRedundantErrors(InFilter));
 
 	// if we have no error messages now, assume success!
 	if (bFoundRedundantError && InCommand.ErrorMessages.Num() == 0 && !InCommand.bCommandSuccessful)
@@ -1651,4 +1844,4 @@ void SwitchVerboseLogs(const bool bInEnable)
 	}
 }
 
-}
\ No newline at end of file
+} // namespace PlasticSourceControlUtils
diff --git a/Engine/Plugins/Developer/PlasticSourceControl/Source/PlasticSourceControl/Private/PlasticSourceControlUtils.h b/Engine/Plugins/Developer/PlasticSourceControl/Source/PlasticSourceControl/Private/PlasticSourceControlUtils.h
--- a/Engine/Plugins/Developer/PlasticSourceControl/Source/PlasticSourceControl/Private/PlasticSourceControlUtils.h
+++ b/Engine/Plugins/Developer/PlasticSourceControl/Source/PlasticSourceControl/Private/PlasticSourceControlUtils.h
@@ -6,8 +6,10 @@
 #include "ISourceControlProvider.h"
 #include "PlasticSourceControlRevision.h"
 
+class FPlasticSourceControlChangelistState;
 class FPlasticSourceControlCommand;
 class FPlasticSourceControlState;
+struct FSoftwareVersion;
 
 /**
  * Helper struct for maintaining temporary files for passing to commands
@@ -15,9 +17,9 @@ class FPlasticSourceControlState;
 class FScopedTempFile
 {
 public:
-
 	/** Constructor - open & write string to temp file */
-	FScopedTempFile(const FText& InText);
+	explicit FScopedTempFile(const FString& InText);
+	explicit FScopedTempFile(const FText& InText);
 
 	/** Destructor - delete temp file */
 	~FScopedTempFile();
@@ -33,6 +35,8 @@ private:
 namespace PlasticSourceControlUtils
 {
 
+const struct FSoftwareVersion& GetOldestSupportedPlasticScmVersion();
+
 /**
  * Find the path to the Plastic binary: for now relying on the Path to access the "cm" command.
  */
@@ -61,7 +65,7 @@ bool FindRootDirectory(const FString& InPathToGameDir, FString& OutWorkspaceRoot
  * Get Plastic SCM cli version
  * @param	OutCliVersion		Version of the Plastic SCM Command Line Interface tool
 */
-void GetPlasticScmVersion(FString& OutPlasticScmVersion);
+bool GetPlasticScmVersion(FSoftwareVersion& OutPlasticScmVersion);
 
 /**
  * Get Plastic SCM current user
@@ -71,18 +75,30 @@ void GetUserName(FString& OutUserName);
 
 /**
  * Get Plastic workspace name
+ * @param	InWorkspaceRoot		The workspace from where to run the command - usually the Game directory
  * @param	OutWorkspaceName	Name of the current workspace
- */
-bool GetWorkspaceName(FString& OutWorkspaceName);
+ * @param	OutErrorMessages	Any errors (from StdErr) as an array per-line
+*/
+bool GetWorkspaceName(const FString& InWorkspaceRoot, FString& OutWorkspaceName, TArray<FString>& OutErrorMessages);
 
 /**
  * Get Plastic repository name and server URL, branch name and current changeset number
  * @param	OutChangeset		The current Changeset Number
  * @param	OutRepositoryName	Name of the repository of the current workspace
- * @param	OutServerUrl		Url/Port of the server of the repository
+ * @param	OutServerUrl		URL/Port of the server of the repository
  * @param	OutBranchName		Name of the current checked-out branch
+ * @param	OutErrorMessages	Any errors (from StdErr) as an array per-line
  */
-bool GetWorkspaceInformation(int32& OutChangeset, FString& OutRepositoryName, FString& OutServerUrl, FString& OutBranchName);
+bool GetWorkspaceInformation(int32& OutChangeset, FString& OutRepositoryName, FString& OutServerUrl, FString& OutBranchName, TArray<FString>& OutErrorMessages);
+
+/**
+ * Use the Project Settings to replace Plastic SCM full username/e-mail by a shorter version for display.
+ *
+ * Used when retrieving the username of a revision, to display in history and content browser asset tooltip.
+ *
+ * @param	InUserName			The Plastic SCM username to shorten for display.
+ */
+FString UserNameToDisplayName(const FString& InUserName);
 
 /**
  * Run a Plastic command - output is a string TArray.
@@ -90,7 +106,7 @@ bool GetWorkspaceInformation(int32& OutChangeset, FString& OutRepositoryName, FS
  * @param	InCommand			The Plastic command - e.g. commit
  * @param	InParameters		The parameters to the Plastic command
  * @param	InFiles				The files to be operated on
- * @param	InConcurrency		Is the command running in the background, or bloking the main thread
+ * @param	InConcurrency		Is the command running in the background, or blocking the main thread
  * @param	OutResults			The results (from StdOut) as an array per-line
  * @param	OutErrorMessages	Any errors (from StdErr) as an array per-line
  * @returns true if the command succeeded and returned no errors
@@ -124,7 +140,7 @@ bool RunUpdateStatus(const TArray<FString>& InFiles, const bool bInUpdateHistory
 bool RunDumpToFile(const FString& InPathToPlasticBinary, const FString& InRevSpec, const FString& InDumpFileName);
 
 /**
- * Run a Plastic "history" and "log" commands and parse it.
+ * Run Plastic "history" and "log" commands and parse their results.
  *
  * @param	bInUpdateHistory	If getting the history of files, versus only checking the heads of branches to detect newer commits
  * @param	InOutStates			The file states to update with the history
@@ -132,6 +148,15 @@ bool RunDumpToFile(const FString& InPathToPlasticBinary, const FString& InRevSpe
  */
 bool RunGetHistory(const bool bInUpdateHistory, TArray<FPlasticSourceControlState>& InOutStates, TArray<FString>& OutErrorMessages);
 
+/**
+ * Run a Plastic "status --changelist --xml" and parse its result.
+ * @param	InConcurrency
+ * @param	OutChangelistsStates	The list of changelists (without their files)
+ * @param	OutCLFilesStates		The list of files per changelist
+ * @param	OutErrorMessages		Any errors (from StdErr) as an array per-line
+ */
+bool RunGetChangelists(const EConcurrency::Type InConcurrency, TArray<FPlasticSourceControlChangelistState>& OutChangelistsStates, TArray<TArray<FPlasticSourceControlState>>& OutCLFilesStates, TArray<FString>& OutErrorMessages);
+
 /**
  * Helper function for various commands to update cached states.
  * @returns true if any states were updated
@@ -157,4 +182,4 @@ void SwitchVerboseLogs(const bool bInEnable);
  */
 FString FindCommonDirectory(const FString& InPath1, const FString& InPath2);
 
-}
+} // namespace PlasticSourceControlUtils
diff --git a/Engine/Plugins/Developer/PlasticSourceControl/Source/PlasticSourceControl/Private/SPlasticSourceControlSettings.cpp b/Engine/Plugins/Developer/PlasticSourceControl/Source/PlasticSourceControl/Private/SPlasticSourceControlSettings.cpp
--- a/Engine/Plugins/Developer/PlasticSourceControl/Source/PlasticSourceControl/Private/SPlasticSourceControlSettings.cpp
+++ b/Engine/Plugins/Developer/PlasticSourceControl/Source/PlasticSourceControl/Private/SPlasticSourceControlSettings.cpp
@@ -31,11 +31,11 @@ void SPlasticSourceControlSettings::Construct(const FArguments& InArgs)
 {
 	FSlateFontInfo Font = FAppStyle::GetFontStyle(TEXT("SourceControl.LoginWindow.Font"));
 
-	bAutoCreateIgnoreFile = true;
+	bAutoCreateIgnoreFile = CanAutoCreateIgnoreFile();
 	bAutoInitialCommit = true;
 
 	InitialCommitMessage = LOCTEXT("InitialCommitMessage", "Initial checkin");
-	ServerUrl = FText::FromString(TEXT("localhost:8087"));
+	ServerUrl = FText::FromString(TEXT("YourOrganization@cloud"));
 	if (FApp::HasProjectName())
 	{
 		WorkspaceName = FText::FromString(FApp::GetProjectName());
@@ -45,6 +45,34 @@ void SPlasticSourceControlSettings::Construct(const FArguments& InArgs)
 	ChildSlot
 	[
 		SNew(SVerticalBox)
+		// Versions (Plugin & Plastic SCM) useful eg to help diagnose issues from screenshots
+		+SVerticalBox::Slot()
+		.AutoHeight()
+		.Padding(2.0f)
+		.VAlign(VAlign_Center)
+		[
+			SNew(SHorizontalBox)
+			.ToolTipText(LOCTEXT("PlasticVersions_Tooltip", "Plastic SCM and Plugin versions"))
+			+SHorizontalBox::Slot()
+			.FillWidth(1.0f)
+			.HAlign(HAlign_Center)
+			[
+				SNew(STextBlock)
+				.Text(this, &SPlasticSourceControlSettings::GetVersions)
+				.Font(Font)
+			]
+		]
+		// Plastic SCM command line tool not available warning
+		+SVerticalBox::Slot()
+		.FillHeight(1.0f)
+		.Padding(2.0f)
+		[
+			SNew(STextBlock)
+			.Visibility(this, &SPlasticSourceControlSettings::PlasticNotAvailable)
+			.ToolTipText(LOCTEXT("PlasticNotAvailable_Tooltip", "Failed to launch Plastic SCM 'cm' command line tool. You need to install it and make sure that 'cm' is on the Path and correctly configured."))
+			.Text(LOCTEXT("PlasticNotAvailable", "Plastic SCM Command Line tool 'cm' failed to start:"))
+			.Font(Font)
+		]
 		// Path to the Plastic SCM binary
 		+SVerticalBox::Slot()
 		.AutoHeight()
@@ -52,12 +80,12 @@ void SPlasticSourceControlSettings::Construct(const FArguments& InArgs)
 		.VAlign(VAlign_Center)
 		[
 			SNew(SHorizontalBox)
-			.ToolTipText(LOCTEXT("BinaryPathLabel_Tooltip", "Path to the Plastic SCM cli binary"))
+			.ToolTipText(LOCTEXT("BinaryPathLabel_Tooltip", "Path to the Plastic SCM Command Line tool 'cm' binary"))
 			+SHorizontalBox::Slot()
 			.FillWidth(1.0f)
 			[
 				SNew(STextBlock)
-				.Text(LOCTEXT("PathLabel", "Plastic SCM Path"))
+				.Text(LOCTEXT("PathLabel", "Plastic SCM Path to cm"))
 				.Font(Font)
 			]
 			+SHorizontalBox::Slot()
@@ -77,7 +105,7 @@ void SPlasticSourceControlSettings::Construct(const FArguments& InArgs)
 		.VAlign(VAlign_Center)
 		[
 			SNew(SHorizontalBox)
-			.ToolTipText(LOCTEXT("WorkspaceRootLabel_Tooltip", "Path to the root of the Plastic SCM workspace"))
+			.ToolTipText(this, &SPlasticSourceControlSettings::GetPathToWorkspaceRoot)
 			+SHorizontalBox::Slot()
 			.FillWidth(1.0f)
 			[
@@ -128,19 +156,12 @@ void SPlasticSourceControlSettings::Construct(const FArguments& InArgs)
 		+SVerticalBox::Slot()
 		.FillHeight(1.0f)
 		.Padding(2.0f)
-		.VAlign(VAlign_Center)
 		[
-			SNew(SHorizontalBox)
+			SNew(STextBlock)
 			.Visibility(this, &SPlasticSourceControlSettings::CanInitializePlasticWorkspace)
 			.ToolTipText(LOCTEXT("WorkspaceNotFound_Tooltip", "No Workspace found at the level or above the current Project. Use the form to create a new one."))
-			+SHorizontalBox::Slot()
-			.FillWidth(1.0f)
-			.HAlign(HAlign_Center)
-			[
-				SNew(STextBlock)
-				.Text(LOCTEXT("WorkspaceNotFound", "Current Project is not in a Plastic SCM Workspace. Create a new one:"))
-				.Font(Font)
-			]
+			.Text(LOCTEXT("WorkspaceNotFound", "Current Project is not in a Plastic SCM Workspace. Create a new one:"))
+			.Font(Font)
 		]
 		// Workspace and Repository Name
 		+SVerticalBox::Slot()
@@ -189,7 +210,7 @@ void SPlasticSourceControlSettings::Construct(const FArguments& InArgs)
 		[
 			SNew(SHorizontalBox)
 			.Visibility(this, &SPlasticSourceControlSettings::CanInitializePlasticWorkspace)
-			.ToolTipText(LOCTEXT("ServerUrl_Tooltip", "Enter the Server URL in the form address:port (eg. localhost:8087 or Organization@cloud)"))
+			.ToolTipText(LOCTEXT("ServerUrl_Tooltip", "Enter the Server URL in the form address:port (eg. YourOrganization@cloud, local, or something like ip:port, eg localhost:8087)"))
 			+SHorizontalBox::Slot()
 			.FillWidth(1.0f)
 			[
@@ -213,19 +234,12 @@ void SPlasticSourceControlSettings::Construct(const FArguments& InArgs)
 		.Padding(2.0f)
 		.VAlign(VAlign_Center)
 		[
-			SNew(SHorizontalBox)
+			SNew(SCheckBox)
 			.Visibility(this, &SPlasticSourceControlSettings::CanInitializePlasticWorkspace)
 			.ToolTipText(LOCTEXT("CreateIgnoreFile_Tooltip", "Create and add a standard 'ignore.conf' file"))
-			+SHorizontalBox::Slot()
-			.FillWidth(0.1f)
-			[
-				SNew(SCheckBox)
-				.IsChecked(ECheckBoxState::Checked)
-				.OnCheckStateChanged(this, &SPlasticSourceControlSettings::OnCheckedCreateIgnoreFile)
-			]
-			+SHorizontalBox::Slot()
-			.FillWidth(2.0f)
-			.VAlign(VAlign_Center)
+			.IsEnabled(this, &SPlasticSourceControlSettings::CanAutoCreateIgnoreFile)
+			.IsChecked(bAutoCreateIgnoreFile)
+			.OnCheckStateChanged(this, &SPlasticSourceControlSettings::OnCheckedCreateIgnoreFile)
 			[
 				SNew(STextBlock)
 				.Text(LOCTEXT("CreateIgnoreFile", "Add a ignore.conf file"))
@@ -242,19 +256,16 @@ void SPlasticSourceControlSettings::Construct(const FArguments& InArgs)
 			.Visibility(this, &SPlasticSourceControlSettings::CanInitializePlasticWorkspace)
 			.ToolTipText(LOCTEXT("InitialCommit_Tooltip", "Make the initial Plastic SCM checkin"))
 			+SHorizontalBox::Slot()
-			.FillWidth(0.1f)
+			.FillWidth(0.7f)
 			[
 				SNew(SCheckBox)
 				.IsChecked(ECheckBoxState::Checked)
 				.OnCheckStateChanged(this, &SPlasticSourceControlSettings::OnCheckedInitialCommit)
-			]
-			+SHorizontalBox::Slot()
-			.FillWidth(0.6f)
-			.VAlign(VAlign_Center)
-			[
-				SNew(STextBlock)
-				.Text(LOCTEXT("InitialCommit", "Initial Checkin"))
-				.Font(Font)
+				[
+					SNew(STextBlock)
+					.Text(LOCTEXT("InitialCommit", "Initial Checkin"))
+					.Font(Font)
+				]
 			]
 			+SHorizontalBox::Slot()
 			.FillWidth(1.4f)
@@ -273,18 +284,10 @@ void SPlasticSourceControlSettings::Construct(const FArguments& InArgs)
 		.Padding(2.0f)
 		.VAlign(VAlign_Center)
 		[
-			SNew(SHorizontalBox)
+			SNew(SCheckBox)
+			.IsChecked(SPlasticSourceControlSettings::IsUpdateStatusAtStartupChecked())
+			.OnCheckStateChanged(this, &SPlasticSourceControlSettings::OnCheckedUpdateStatusAtStartup)
 			.ToolTipText(LOCTEXT("UpdateStatusAtStartup_Tooltip", "Run an asynchronous Update Status at Editor startup (can be slow)."))
-			+SHorizontalBox::Slot()
-			.FillWidth(0.1f)
-			[
-				SNew(SCheckBox)
-				.IsChecked(SPlasticSourceControlSettings::IsUpdateStatusAtStartupChecked())
-				.OnCheckStateChanged(this, &SPlasticSourceControlSettings::OnCheckedUpdateStatusAtStartup)
-			]
-			+SHorizontalBox::Slot()
-			.FillWidth(2.0f)
-			.VAlign(VAlign_Center)
 			[
 				SNew(STextBlock)
 				.Text(LOCTEXT("UpdateStatusAtStartup", "Update workspace Status at Editor startup"))
@@ -297,18 +300,10 @@ void SPlasticSourceControlSettings::Construct(const FArguments& InArgs)
 		.Padding(2.0f)
 		.VAlign(VAlign_Center)
 		[
-			SNew(SHorizontalBox)
-			.ToolTipText(LOCTEXT("UpdateStatusOtherBranches_Tooltip", "Enable Update status to detect changesets on other branches (can be slow)."))
-			+SHorizontalBox::Slot()
-			.FillWidth(0.1f)
-			[
-				SNew(SCheckBox)
-				.IsChecked(SPlasticSourceControlSettings::IsUpdateStatusOtherBranchesChecked())
-				.OnCheckStateChanged(this, &SPlasticSourceControlSettings::OnCheckedUpdateStatusOtherBranches)
-			]
-			+SHorizontalBox::Slot()
-			.FillWidth(2.0f)
-			.VAlign(VAlign_Center)
+			SNew(SCheckBox)
+			.ToolTipText(LOCTEXT("UpdateStatusOtherBranches_Tooltip", "Enable Update status to detect more recent changes on other branches in order to display warnings (can be slow)."))
+			.IsChecked(SPlasticSourceControlSettings::IsUpdateStatusOtherBranchesChecked())
+			.OnCheckStateChanged(this, &SPlasticSourceControlSettings::OnCheckedUpdateStatusOtherBranches)
 			[
 				SNew(STextBlock)
 				.Text(LOCTEXT("UpdateStatusOtherBranches", "Update status also check branch history."))
@@ -321,18 +316,10 @@ void SPlasticSourceControlSettings::Construct(const FArguments& InArgs)
 		.Padding(2.0f)
 		.VAlign(VAlign_Center)
 		[
-			SNew(SHorizontalBox)
+			SNew(SCheckBox)
 			.ToolTipText(LOCTEXT("EnableVerboseLogs_Tooltip", "Override LogSourceControl default verbosity level to Verbose (except if already set to VeryVerbose)."))
-			+SHorizontalBox::Slot()
-			.FillWidth(0.1f)
-			[
-				SNew(SCheckBox)
-				.IsChecked(SPlasticSourceControlSettings::IsEnableVerboseLogsChecked())
-				.OnCheckStateChanged(this, &SPlasticSourceControlSettings::OnCheckedEnableVerboseLogs)
-			]
-			+SHorizontalBox::Slot()
-			.FillWidth(2.0f)
-			.VAlign(VAlign_Center)
+			.IsChecked(SPlasticSourceControlSettings::IsEnableVerboseLogsChecked())
+			.OnCheckStateChanged(this, &SPlasticSourceControlSettings::OnCheckedEnableVerboseLogs)
 			[
 				SNew(STextBlock)
 				.Text(LOCTEXT("EnableVerboseLogs", "Enable Source Control Verbose logs"))
@@ -380,45 +367,54 @@ void SPlasticSourceControlSettings::Construct(const FArguments& InArgs)
 	];
 }
 
+EVisibility SPlasticSourceControlSettings::PlasticNotAvailable() const
+{
+	const FPlasticSourceControlProvider& Provider = FPlasticSourceControlModule::Get().GetProvider();
+	return Provider.IsPlasticAvailable() ? EVisibility::Collapsed : EVisibility::Visible;
+}
+
 FText SPlasticSourceControlSettings::GetBinaryPathText() const
 {
-	const FPlasticSourceControlModule& PlasticSourceControl = FModuleManager::LoadModuleChecked<FPlasticSourceControlModule>("PlasticSourceControl");
-	return FText::FromString(PlasticSourceControl.AccessSettings().GetBinaryPath());
+	return FText::FromString(FPlasticSourceControlModule::Get().GetProvider().AccessSettings().GetBinaryPath());
 }
 
 void SPlasticSourceControlSettings::OnBinaryPathTextCommited(const FText& InText, ETextCommit::Type InCommitType) const
 {
-	FPlasticSourceControlModule& PlasticSourceControl = FModuleManager::LoadModuleChecked<FPlasticSourceControlModule>("PlasticSourceControl");
-	const bool bChanged = PlasticSourceControl.AccessSettings().SetBinaryPath(InText.ToString());
+	FPlasticSourceControlProvider& Provider = FPlasticSourceControlModule::Get().GetProvider();
+	const bool bChanged = Provider.AccessSettings().SetBinaryPath(InText.ToString());
 	if (bChanged)
 	{
 		// Re-Check provided Plastic binary path for each change
-		PlasticSourceControl.GetProvider().CheckPlasticAvailability();
-		if (PlasticSourceControl.GetProvider().IsPlasticAvailable())
+		Provider.CheckPlasticAvailability();
+		if (Provider.IsPlasticAvailable())
 		{
-			PlasticSourceControl.SaveSettings();
+			Provider.AccessSettings().SaveSettings();
 		}
 	}
 }
 
+FText SPlasticSourceControlSettings::GetVersions() const
+{
+	const FPlasticSourceControlProvider& Provider = FPlasticSourceControlModule::Get().GetProvider();
+	return FText::FromString(TEXT("Plastic SCM ") + Provider.GetPlasticScmVersion().String + TEXT(" (plugin v") + Provider.GetPluginVersion() + TEXT(")"));
+}
+
 FText SPlasticSourceControlSettings::GetPathToWorkspaceRoot() const
 {
-   const FPlasticSourceControlModule& PlasticSourceControl = FModuleManager::LoadModuleChecked<FPlasticSourceControlModule>("PlasticSourceControl");
-	return FText::FromString(PlasticSourceControl.GetProvider().GetPathToWorkspaceRoot());
+	return FText::FromString(FPlasticSourceControlModule::Get().GetProvider().GetPathToWorkspaceRoot());
 }
 
 FText SPlasticSourceControlSettings::GetUserName() const
 {
-   const FPlasticSourceControlModule& PlasticSourceControl = FModuleManager::LoadModuleChecked<FPlasticSourceControlModule>("PlasticSourceControl");
-	return FText::FromString(PlasticSourceControl.GetProvider().GetUserName());
+	return FText::FromString(FPlasticSourceControlModule::Get().GetProvider().GetUserName());
 }
 
 
 EVisibility SPlasticSourceControlSettings::CanInitializePlasticWorkspace() const
 {
-   const FPlasticSourceControlModule& PlasticSourceControl = FModuleManager::LoadModuleChecked<FPlasticSourceControlModule>("PlasticSourceControl");
-	const bool bPlasticAvailable = PlasticSourceControl.GetProvider().IsPlasticAvailable();
-	const bool bPlasticWorkspaceFound = PlasticSourceControl.GetProvider().IsWorkspaceFound();
+	const FPlasticSourceControlProvider& Provider = FPlasticSourceControlModule::Get().GetProvider();
+	const bool bPlasticAvailable = Provider.IsPlasticAvailable();
+	const bool bPlasticWorkspaceFound = Provider.IsWorkspaceFound();
 	return (bPlasticAvailable && !bPlasticWorkspaceFound) ? EVisibility::Visible : EVisibility::Collapsed;
 }
 
@@ -461,6 +457,12 @@ FText SPlasticSourceControlSettings::GetServerUrl() const
 	return ServerUrl;
 }
 
+bool SPlasticSourceControlSettings::CanAutoCreateIgnoreFile() const
+{
+	const bool bIgnoreFileFound = FPaths::FileExists(GetIgnoreFileName());
+	return !bIgnoreFileFound;
+}
+
 void SPlasticSourceControlSettings::OnCheckedCreateIgnoreFile(ECheckBoxState NewCheckedState)
 {
 	bAutoCreateIgnoreFile = (NewCheckedState == ECheckBoxState::Checked);
@@ -502,8 +504,8 @@ void SPlasticSourceControlSettings::LaunchMakeWorkspaceOperation()
 	MakeWorkspaceOperation->RepositoryName = RepositoryName.ToString();
 	MakeWorkspaceOperation->ServerUrl = ServerUrl.ToString();
 
-	FPlasticSourceControlModule& PlasticSourceControl = FModuleManager::LoadModuleChecked<FPlasticSourceControlModule>("PlasticSourceControl");
-	ECommandResult::Type Result = PlasticSourceControl.GetProvider().Execute(MakeWorkspaceOperation, TArray<FString>(), EConcurrency::Asynchronous, FSourceControlOperationComplete::CreateSP(this, &SPlasticSourceControlSettings::OnSourceControlOperationComplete));
+	FPlasticSourceControlProvider& Provider = FPlasticSourceControlModule::Get().GetProvider();
+	ECommandResult::Type Result = Provider.Execute(MakeWorkspaceOperation, TArray<FString>(), EConcurrency::Asynchronous, FSourceControlOperationComplete::CreateSP(this, &SPlasticSourceControlSettings::OnSourceControlOperationComplete));
 	if (Result == ECommandResult::Succeeded)
 	{
 		DisplayInProgressNotification(MakeWorkspaceOperation->GetInProgressString());
@@ -514,25 +516,25 @@ void SPlasticSourceControlSettings::LaunchMakeWorkspaceOperation()
 	}
 }
 
-/// 2. Add all project files to to Source Control (.uproject, Config/, Content/, Source/ files and ignore.conf if any)
+/// 2. Add all project files to Source Control (.uproject, Config/, Content/, Source/ files and ignore.conf if any)
 void SPlasticSourceControlSettings::LaunchMarkForAddOperation()
 {
 	TSharedRef<FMarkForAdd, ESPMode::ThreadSafe> MarkForAddOperation = ISourceControlOperation::Create<FMarkForAdd>();
-	FPlasticSourceControlModule& PlasticSourceControl = FModuleManager::LoadModuleChecked<FPlasticSourceControlModule>("PlasticSourceControl");
+	FPlasticSourceControlProvider& Provider = FPlasticSourceControlModule::Get().GetProvider();
 
 	// 1.b. Check the new workspace status to enable connection
-	PlasticSourceControl.GetProvider().CheckPlasticAvailability();
+	Provider.CheckPlasticAvailability();
 
-	if (PlasticSourceControl.GetProvider().IsWorkspaceFound())
+	if (Provider.IsWorkspaceFound())
 	{
 		if (bAutoCreateIgnoreFile)
 		{
 			// 1.c Create a standard "ignore.conf" file with common patterns for a typical Blueprint & C++ project
 			CreateIgnoreFile();
 		}
-		// 2. Add all project files to to Source Control (.uproject, Config/, Content/, Source/ files and ignore.conf if any)
+		// 2. Add all project files to Source Control (.uproject, Config/, Content/, Source/ files and ignore.conf if any)
 		const TArray<FString> ProjectFiles = GetProjectFiles();
-		ECommandResult::Type Result = PlasticSourceControl.GetProvider().Execute(MarkForAddOperation, ProjectFiles, EConcurrency::Asynchronous, FSourceControlOperationComplete::CreateSP(this, &SPlasticSourceControlSettings::OnSourceControlOperationComplete));
+		ECommandResult::Type Result = Provider.Execute(MarkForAddOperation, ProjectFiles, EConcurrency::Asynchronous, FSourceControlOperationComplete::CreateSP(this, &SPlasticSourceControlSettings::OnSourceControlOperationComplete));
 		if (Result == ECommandResult::Succeeded)
 		{
 			DisplayInProgressNotification(MarkForAddOperation->GetInProgressString());
@@ -553,9 +555,9 @@ void SPlasticSourceControlSettings::LaunchCheckInOperation()
 {
 	TSharedRef<FCheckIn, ESPMode::ThreadSafe> CheckInOperation = ISourceControlOperation::Create<FCheckIn>();
 	CheckInOperation->SetDescription(InitialCommitMessage);
-	FPlasticSourceControlModule& PlasticSourceControl = FModuleManager::LoadModuleChecked<FPlasticSourceControlModule>("PlasticSourceControl");
+	FPlasticSourceControlProvider& Provider = FPlasticSourceControlModule::Get().GetProvider();
 	const TArray<FString> ProjectFiles = GetProjectFiles(); // Note: listing files and folders is only needed for the update status operation following the checkin to know on what to operate
-	ECommandResult::Type Result = PlasticSourceControl.GetProvider().Execute(CheckInOperation, ProjectFiles, EConcurrency::Asynchronous, FSourceControlOperationComplete::CreateSP(this, &SPlasticSourceControlSettings::OnSourceControlOperationComplete));
+	ECommandResult::Type Result = Provider.Execute(CheckInOperation, ProjectFiles, EConcurrency::Asynchronous, FSourceControlOperationComplete::CreateSP(this, &SPlasticSourceControlSettings::OnSourceControlOperationComplete));
 	if (Result == ECommandResult::Succeeded)
 	{
 		DisplayInProgressNotification(CheckInOperation->GetInProgressString());
@@ -581,7 +583,7 @@ void SPlasticSourceControlSettings::OnSourceControlOperationComplete(const FSour
 		DisplayFailureNotification(InOperation->GetName());
 	}
 
-	// Launch the following asynchrounous operation
+	// Launch the following asynchronous operation
 	if ((InOperation->GetName() == "MakeWorkspace") && (InResult == ECommandResult::Succeeded) && bAutoInitialCommit)
 	{
 		// 2. Add .uproject, Config/, Content/ and Source/ files (and ignore.conf if any)
@@ -645,8 +647,7 @@ void SPlasticSourceControlSettings::DisplayFailureNotification(const FName& InOp
 /** Delegate to check for presence of a Plastic ignore.conf file to an existing Plastic SCM workspace */
 EVisibility SPlasticSourceControlSettings::CanAddIgnoreFile() const
 {
-   const FPlasticSourceControlModule& PlasticSourceControl = FModuleManager::LoadModuleChecked<FPlasticSourceControlModule>("PlasticSourceControl");
-	const bool bPlasticWorkspaceFound = PlasticSourceControl.GetProvider().IsWorkspaceFound();
+	const bool bPlasticWorkspaceFound = FPlasticSourceControlModule::Get().GetProvider().IsWorkspaceFound();
 	const bool bIgnoreFileFound = FPaths::FileExists(GetIgnoreFileName());
 	return (bPlasticWorkspaceFound && !bIgnoreFileFound) ? EVisibility::Visible : EVisibility::Collapsed;
 }
@@ -670,57 +671,57 @@ FReply SPlasticSourceControlSettings::OnClickedAddIgnoreFile() const
 
 void SPlasticSourceControlSettings::OnCheckedUpdateStatusAtStartup(ECheckBoxState NewCheckedState)
 {
-	FPlasticSourceControlModule& PlasticSourceControl = FModuleManager::GetModuleChecked<FPlasticSourceControlModule>("PlasticSourceControl");
-	PlasticSourceControl.AccessSettings().SetUpdateStatusAtStartup(NewCheckedState == ECheckBoxState::Checked);
-	PlasticSourceControl.AccessSettings().SaveSettings();
+	FPlasticSourceControlSettings& PlasticSettings = FPlasticSourceControlModule::Get().GetProvider().AccessSettings();
+	PlasticSettings.SetUpdateStatusAtStartup(NewCheckedState == ECheckBoxState::Checked);
+	PlasticSettings.SaveSettings();
 }
 
 ECheckBoxState SPlasticSourceControlSettings::IsUpdateStatusAtStartupChecked() const
 {
-	const FPlasticSourceControlModule& PlasticSourceControl = FModuleManager::GetModuleChecked<FPlasticSourceControlModule>("PlasticSourceControl");
-	return PlasticSourceControl.AccessSettings().GetUpdateStatusAtStartup() ? ECheckBoxState::Checked : ECheckBoxState::Unchecked;
+	FPlasticSourceControlSettings& PlasticSettings = FPlasticSourceControlModule::Get().GetProvider().AccessSettings();
+	return PlasticSettings.GetUpdateStatusAtStartup() ? ECheckBoxState::Checked : ECheckBoxState::Unchecked;
 }
 
 void SPlasticSourceControlSettings::OnCheckedUpdateStatusOtherBranches(ECheckBoxState NewCheckedState)
 {
-	FPlasticSourceControlModule& PlasticSourceControl = FModuleManager::GetModuleChecked<FPlasticSourceControlModule>("PlasticSourceControl");
-	PlasticSourceControl.AccessSettings().SetUpdateStatusOtherBranches(NewCheckedState == ECheckBoxState::Checked);
-	PlasticSourceControl.AccessSettings().SaveSettings();
+	FPlasticSourceControlSettings& PlasticSettings = FPlasticSourceControlModule::Get().GetProvider().AccessSettings();
+	PlasticSettings.SetUpdateStatusOtherBranches(NewCheckedState == ECheckBoxState::Checked);
+	PlasticSettings.SaveSettings();
 }
 
 ECheckBoxState SPlasticSourceControlSettings::IsUpdateStatusOtherBranchesChecked() const
 {
-	const FPlasticSourceControlModule& PlasticSourceControl = FModuleManager::GetModuleChecked<FPlasticSourceControlModule>("PlasticSourceControl");
-	return PlasticSourceControl.AccessSettings().GetUpdateStatusOtherBranches() ? ECheckBoxState::Checked : ECheckBoxState::Unchecked;
+	FPlasticSourceControlSettings& PlasticSettings = FPlasticSourceControlModule::Get().GetProvider().AccessSettings();
+	return PlasticSettings.GetUpdateStatusOtherBranches() ? ECheckBoxState::Checked : ECheckBoxState::Unchecked;
 }
 
 void SPlasticSourceControlSettings::OnCheckedEnableVerboseLogs(ECheckBoxState NewCheckedState)
 {
-	FPlasticSourceControlModule& PlasticSourceControl = FModuleManager::GetModuleChecked<FPlasticSourceControlModule>("PlasticSourceControl");
-	PlasticSourceControl.AccessSettings().SetEnableVerboseLogs(NewCheckedState == ECheckBoxState::Checked);
-	PlasticSourceControl.AccessSettings().SaveSettings();
+	FPlasticSourceControlSettings& PlasticSettings = FPlasticSourceControlModule::Get().GetProvider().AccessSettings();
+	PlasticSettings.SetEnableVerboseLogs(NewCheckedState == ECheckBoxState::Checked);
+	PlasticSettings.SaveSettings();
 
 	PlasticSourceControlUtils::SwitchVerboseLogs(NewCheckedState == ECheckBoxState::Checked);
 }
 
 ECheckBoxState SPlasticSourceControlSettings::IsEnableVerboseLogsChecked() const
 {
-	const FPlasticSourceControlModule& PlasticSourceControl = FModuleManager::GetModuleChecked<FPlasticSourceControlModule>("PlasticSourceControl");
-	return PlasticSourceControl.AccessSettings().GetEnableVerboseLogs() ? ECheckBoxState::Checked : ECheckBoxState::Unchecked;
+	FPlasticSourceControlSettings& PlasticSettings = FPlasticSourceControlModule::Get().GetProvider().AccessSettings();
+	return PlasticSettings.GetEnableVerboseLogs() ? ECheckBoxState::Checked : ECheckBoxState::Unchecked;
 }
 
 /** Path to the "ignore.conf" file */
-const FString& SPlasticSourceControlSettings::GetIgnoreFileName() const
+const FString SPlasticSourceControlSettings::GetIgnoreFileName() const
 {
-	static const FString PathToProjectDir = FPaths::ConvertRelativePathToFull(FPaths::ProjectDir());
-	static const FString IgnoreFileName = FPaths::Combine(*PathToProjectDir, TEXT("ignore.conf"));
+	const FString PathToWorkspaceRoot = FPlasticSourceControlModule::Get().GetProvider().GetPathToWorkspaceRoot();
+	const FString IgnoreFileName = FPaths::Combine(*PathToWorkspaceRoot, TEXT("ignore.conf"));
 	return IgnoreFileName;
 }
 
 /** Create a standard "ignore.conf" file with common patterns for a typical Blueprint & C++ project */
 bool SPlasticSourceControlSettings::CreateIgnoreFile() const
 {
-	const FString IgnoreFileContent = TEXT("Binaries\nBuild\nDerivedDataCache\nIntermediate\nSaved\n.vscode\n.vs\n*.VC.db\n*.opensdf\n*.opendb\n*.sdf\n*.sln\n*.suo\n*.xcodeproj\n*.xcworkspace");
+	const FString IgnoreFileContent = TEXT("Binaries\nBuild\nDerivedDataCache\nIntermediate\nSaved\nScript\nenc_temp_folder\n.idea\n.vscode\n.vs\n*.VC.db\n*.opensdf\n*.opendb\n*.sdf\n*.sln\n*.suo\n*.xcodeproj\n*.xcworkspace");
 	return FFileHelper::SaveStringToFile(IgnoreFileContent, *GetIgnoreFileName(), FFileHelper::EEncodingOptions::ForceUTF8WithoutBOM);
 }
 
diff --git a/Engine/Plugins/Developer/PlasticSourceControl/Source/PlasticSourceControl/Private/SPlasticSourceControlSettings.h b/Engine/Plugins/Developer/PlasticSourceControl/Source/PlasticSourceControl/Private/SPlasticSourceControlSettings.h
--- a/Engine/Plugins/Developer/PlasticSourceControl/Source/PlasticSourceControl/Private/SPlasticSourceControlSettings.h
+++ b/Engine/Plugins/Developer/PlasticSourceControl/Source/PlasticSourceControl/Private/SPlasticSourceControlSettings.h
@@ -15,16 +15,15 @@
 class SPlasticSourceControlSettings : public SCompoundWidget
 {
 public:
-	
 	SLATE_BEGIN_ARGS(SPlasticSourceControlSettings) {}
-	
+
 	SLATE_END_ARGS()
 
 public:
-
 	void Construct(const FArguments& InArgs);
 
 private:
+	EVisibility PlasticNotAvailable() const;
 
 	/** Delegate to get cm binary path from settings */
 	FText GetBinaryPathText() const;
@@ -33,6 +32,7 @@ private:
 	void OnBinaryPathTextCommited(const FText& InText, ETextCommit::Type InCommitType) const;
 
 	/** Delegate to get workspace root and user name from provider */
+	FText GetVersions() const;
 	FText GetPathToWorkspaceRoot() const;
 	FText GetUserName() const;
 
@@ -48,6 +48,7 @@ private:
 	void OnServerUrlCommited(const FText& InText, ETextCommit::Type InCommitType);
 	FText GetServerUrl() const;
 	FText ServerUrl;
+	bool CanAutoCreateIgnoreFile() const;
 	void OnCheckedCreateIgnoreFile(ECheckBoxState NewCheckedState);
 	bool bAutoCreateIgnoreFile;
 
@@ -88,7 +89,7 @@ private:
 	EVisibility CanAddIgnoreFile() const;
 	FReply OnClickedAddIgnoreFile() const;
 
-	const FString& GetIgnoreFileName() const;
+	const FString GetIgnoreFileName() const;
 	bool CreateIgnoreFile() const;
 
 	TArray<FString> GetProjectFiles() const;
diff --git a/Engine/Plugins/Developer/PlasticSourceControl/Source/PlasticSourceControl/Private/SoftwareVersion.cpp b/Engine/Plugins/Developer/PlasticSourceControl/Source/PlasticSourceControl/Private/SoftwareVersion.cpp
new file mode 100644
--- /dev/null
+++ b/Engine/Plugins/Developer/PlasticSourceControl/Source/PlasticSourceControl/Private/SoftwareVersion.cpp
@@ -0,0 +1,44 @@
+// Copyright Epic Games, Inc. All Rights Reserved.
+
+#include "SoftwareVersion.h"
+
+FSoftwareVersion::FSoftwareVersion(const FString& InVersionString)
+{
+	String = InVersionString;
+	TArray<FString> Parts;
+	const int32 N = String.ParseIntoArray(Parts, TEXT("."));
+	if (N == 4)
+	{
+		Major = FCString::Atoi(*Parts[0]);
+		Minor = FCString::Atoi(*Parts[1]);
+		Patch = FCString::Atoi(*Parts[2]);
+		Changeset = FCString::Atoi(*Parts[3]);
+	}
+}
+
+FSoftwareVersion::FSoftwareVersion(const int& InMajor, const int& InMinor, const int& InPatch, const int& InChangeset)
+{
+	String = FString::Printf(TEXT("%d.%d.%d.%d"), InMajor, InMinor, InPatch);
+	Major = InMajor;
+	Minor = InMinor;
+	Patch = InPatch;
+	Changeset = InChangeset;
+}
+
+bool operator==(const FSoftwareVersion& Rhs, const FSoftwareVersion& Lhs)
+{
+	return (Rhs.Major == Lhs.Major) && (Rhs.Minor == Lhs.Minor) && (Rhs.Patch == Lhs.Patch) && (Rhs.Changeset == Lhs.Changeset);
+}
+
+bool operator<(const FSoftwareVersion& Rhs, const FSoftwareVersion& Lhs)
+{
+	if (Rhs.Major < Lhs.Major) return true;
+	if (Rhs.Major > Lhs.Major) return false;
+	if (Rhs.Minor < Lhs.Minor) return true;
+	if (Rhs.Minor > Lhs.Minor) return false;
+	if (Rhs.Patch < Lhs.Patch) return true;
+	if (Rhs.Patch > Lhs.Patch) return false;
+	if (Rhs.Changeset < Lhs.Changeset) return true;
+	if (Rhs.Changeset > Lhs.Changeset) return false;
+	return false; // Equal
+}
diff --git a/Engine/Plugins/Developer/PlasticSourceControl/Source/PlasticSourceControl/Private/SoftwareVersion.h b/Engine/Plugins/Developer/PlasticSourceControl/Source/PlasticSourceControl/Private/SoftwareVersion.h
new file mode 100644
--- /dev/null
+++ b/Engine/Plugins/Developer/PlasticSourceControl/Source/PlasticSourceControl/Private/SoftwareVersion.h
@@ -0,0 +1,26 @@
+// Copyright Epic Games, Inc. All Rights Reserved.
+
+#pragma once
+
+#include "CoreMinimal.h"
+
+/**
+ * Software version string in the form "X.Y.Z.C", ie Major.Minor.Patch.Changeset (as returned by GetPlasticScmVersion)
+*/
+struct FSoftwareVersion
+{
+	FSoftwareVersion() {}
+
+	explicit FSoftwareVersion(const FString& InVersionString);
+	explicit FSoftwareVersion(const int& InMajor, const int& InMinor, const int& InPatch, const int& InChangeset);
+
+	FString String = TEXT("<unknown-version>");
+
+	int Major = 0;
+	int Minor = 0;
+	int Patch = 0;
+	int Changeset = 0;
+};
+
+bool operator==(const FSoftwareVersion& Rhs, const FSoftwareVersion& Lhs);
+bool operator<(const FSoftwareVersion& Rhs, const FSoftwareVersion& Lhs);
diff --git a/Engine/Source/Editor/EditorStyle/Private/StarshipStyle.cpp b/Engine/Source/Editor/EditorStyle/Private/StarshipStyle.cpp
--- a/Engine/Source/Editor/EditorStyle/Private/StarshipStyle.cpp
+++ b/Engine/Source/Editor/EditorStyle/Private/StarshipStyle.cpp
@@ -7106,18 +7106,18 @@ void FStarshipEditorStyle::FStyle::SetupSourceControlStyles()
 	// Plastic SCM
 	{
 		Set("Plastic.CheckedOut", new CORE_IMAGE_BRUSH_SVG("Starship/SourceControl/SCC_CheckedOut", Icon16x16, FStyleColors::AccentRed));
-		Set("Plastic.Changed", new CORE_IMAGE_BRUSH_SVG("Starship/SourceControl/SCC_CheckedOut", Icon16x16, FStyleColors::AccentYellow)); // custom
+		Set("Plastic.Changed", new CORE_IMAGE_BRUSH_SVG("Starship/SourceControl/SCC_CheckedOut", Icon16x16, FStyleColors::AccentWhite)); // custom
 		Set("Plastic.OpenForAdd", new CORE_IMAGE_BRUSH_SVG("Starship/SourceControl/SCC_ContentAdd", Icon16x16, FStyleColors::AccentRed));
 		Set("Plastic.CheckedOutByOtherUser", new CORE_IMAGE_BRUSH_SVG("Starship/SourceControl/SCC_CheckedOut", Icon16x16, FStyleColors::AccentYellow));
 		Set("Plastic.ModifiedOtherBranch", new CORE_IMAGE_BRUSH_SVG("Starship/SourceControl/SCC_ModifiedOtherBranch", Icon16x16, FStyleColors::AccentRed));
 		Set("Plastic.MarkedForDelete", new CORE_IMAGE_BRUSH_SVG("Starship/SourceControl/SCC_MarkedForDelete", Icon16x16, FStyleColors::AccentRed));
-		Set("Plastic.LocallyDeleted", new CORE_IMAGE_BRUSH_SVG("Starship/SourceControl/SCC_MarkedForDelete", Icon16x16, FStyleColors::AccentYellow)); // custom
+		Set("Plastic.LocallyDeleted", new CORE_IMAGE_BRUSH_SVG("Starship/SourceControl/SCC_MarkedForDelete", Icon16x16, FStyleColors::AccentWhite)); // custom
 		Set("Plastic.NotAtHeadRevision", new CORE_IMAGE_BRUSH_SVG("Starship/SourceControl/SCC_ModifiedOtherBranch", Icon16x16, FStyleColors::AccentYellow));
 		Set("Plastic.Conflicted", new CORE_IMAGE_BRUSH_SVG("Starship/SourceControl/SCC_ModifiedOtherBranch", Icon16x16, FStyleColors::AccentRed)); // custom
 		Set("Plastic.NotInDepot", new CORE_IMAGE_BRUSH_SVG("Starship/SourceControl/SCC_NotInDepot", Icon16x16, FStyleColors::AccentYellow));
-		Set("Plastic.Ignored", new CORE_IMAGE_BRUSH_SVG("Starship/SourceControl/SCC_NotInDepot", Icon16x16, FStyleColors::AccentYellow)); // custom
+		Set("Plastic.Ignored", new CORE_IMAGE_BRUSH_SVG("Starship/SourceControl/SCC_NotInDepot", Icon16x16, FStyleColors::AccentWhite)); // custom
 		Set("Plastic.Branched", new CORE_IMAGE_BRUSH_SVG("Starship/SourceControl/SCC_Branched", Icon16x16, FStyleColors::AccentGreen));
-		Set("Plastic.LocallyMoved", new CORE_IMAGE_BRUSH_SVG("Starship/SourceControl/SCC_Branched", Icon16x16, FStyleColors::AccentYellow)); // custom
+		Set("Plastic.LocallyMoved", new CORE_IMAGE_BRUSH_SVG("Starship/SourceControl/SCC_Branched", Icon16x16, FStyleColors::AccentWhite)); // custom
 	}
 	// Subversion
 	{
