PR #9553: [Source Control] Plastic SCM implements ISourceControlChangelist::CanDelete() (Contributed by [REDACTED_PERSON])

diff --git a/Engine/Plugins/Developer/PlasticSourceControl/Source/PlasticSourceControl/Private/PlasticSourceControlChangelist.h b/Engine/Plugins/Developer/PlasticSourceControl/Source/PlasticSourceControl/Private/PlasticSourceControlChangelist.h
--- a/Engine/Plugins/Developer/PlasticSourceControl/Source/PlasticSourceControl/Private/PlasticSourceControlChangelist.h
+++ b/Engine/Plugins/Developer/PlasticSourceControl/Source/PlasticSourceControl/Private/PlasticSourceControlChangelist.h
@@ -22,6 +22,11 @@ public:
 
 	FPlasticSourceControlChangelist(const FPlasticSourceControlChangelist& InOther) = default;
 
+	virtual bool CanDelete() const override
+	{
+		return !IsDefault();
+	}
+
 	bool operator==(const FPlasticSourceControlChangelist& InOther) const
 	{
 		return ChangelistName == InOther.ChangelistName;
diff --git a/Engine/Plugins/Developer/PlasticSourceControl/Source/PlasticSourceControl/Private/PlasticSourceControlOperations.cpp b/Engine/Plugins/Developer/PlasticSourceControl/Source/PlasticSourceControl/Private/PlasticSourceControlOperations.cpp
--- a/Engine/Plugins/Developer/PlasticSourceControl/Source/PlasticSourceControl/Private/PlasticSourceControlOperations.cpp
+++ b/Engine/Plugins/Developer/PlasticSourceControl/Source/PlasticSourceControl/Private/PlasticSourceControlOperations.cpp
@@ -17,8 +17,6 @@
 #include "ISourceControlModule.h"
 #include "Algo/NoneOf.h"
 
-#include "Runtime/Launch/Resources/Version.h"
-
 #define LOCTEXT_NAMESPACE "PlasticSourceControl"
 
 // 11.0.16.7248 add support for --descriptionfile for multi-line descriptions and support for special characters
@@ -48,13 +46,11 @@ void IPlasticSourceControlWorker::RegisterWorkers(FPlasticSourceControlProvider&
 	PlasticSourceControlProvider.RegisterWorker("Copy", FGetPlasticSourceControlWorker::CreateStatic(&InstantiateWorker<FPlasticCopyWorker>));
 	PlasticSourceControlProvider.RegisterWorker("Resolve", FGetPlasticSourceControlWorker::CreateStatic(&InstantiateWorker<FPlasticResolveWorker>));
 
-#if ENGINE_MAJOR_VERSION == 5
 	PlasticSourceControlProvider.RegisterWorker("UpdateChangelistsStatus", FGetPlasticSourceControlWorker::CreateStatic(&InstantiateWorker<FPlasticGetPendingChangelistsWorker>));
 	PlasticSourceControlProvider.RegisterWorker("NewChangelist", FGetPlasticSourceControlWorker::CreateStatic(&InstantiateWorker<FPlasticNewChangelistWorker>));
 	PlasticSourceControlProvider.RegisterWorker("DeleteChangelist", FGetPlasticSourceControlWorker::CreateStatic(&InstantiateWorker<FPlasticDeleteChangelistWorker>));
 	PlasticSourceControlProvider.RegisterWorker("EditChangelist", FGetPlasticSourceControlWorker::CreateStatic(&InstantiateWorker<FPlasticEditChangelistWorker>));
 	PlasticSourceControlProvider.RegisterWorker("MoveToChangelist", FGetPlasticSourceControlWorker::CreateStatic(&InstantiateWorker<FPlasticReopenWorker>));
-#endif
 }
 
 
