PR #10129: [Source Control] Update the Plastic SCM plugin to version 1.8.0 (from 1.6.2) (Contributed by [REDACTED_PERSON])

diff --git a/Engine/Plugins/Developer/PlasticSourceControl/PlasticSourceControl.uplugin b/Engine/Plugins/Developer/PlasticSourceControl/PlasticSourceControl.uplugin
--- a/Engine/Plugins/Developer/PlasticSourceControl/PlasticSourceControl.uplugin
+++ b/Engine/Plugins/Developer/PlasticSourceControl/PlasticSourceControl.uplugin
@@ -1,15 +1,15 @@
 {
 	"FileVersion": 3,
-	"Version": 62,
-	"VersionName": "1.6.2",
+	"Version": 80,
+	"VersionName": "1.8.0",
 	"FriendlyName": "Plastic SCM",
 	"Description": "Plastic source control management",
 	"Category": "Source Control",
 	"CreatedBy": "SRombauts for Codice Software.",
-	"CreatedByURL": "http://srombauts.github.com",
+	"CreatedByURL": "https://srombauts.github.io",
 	"DocsURL": "[REDACTED_PROFILE_URL]/UEPlasticPlugin",
 	"MarketplaceURL": "",
-	"SupportURL": "https://www.plasticscm.com/support",
+	"SupportURL": "https://plasticscm.com/support",
 	"EnabledByDefault": true,
 	"CanContainContent": false,
 	"IsBetaVersion": false,
@@ -17,10 +17,11 @@
 	"Modules": [
 		{
 			"Name": "PlasticSourceControl",
-			"Type": "Editor",
-			"LoadingPhase": "Default",
+			"Type": "UncookedOnly",
+			"LoadingPhase": "EarliestPossible",
 			"PlatformAllowList": [
-				"Win64"
+				"Win64",
+				"Linux"
 			]
 		}
 	]
diff --git a/Engine/Plugins/Developer/PlasticSourceControl/Source/PlasticSourceControl/Private/PackageUtils.cpp b/Engine/Plugins/Developer/PlasticSourceControl/Source/PlasticSourceControl/Private/PackageUtils.cpp
new file mode 100644
--- /dev/null
+++ b/Engine/Plugins/Developer/PlasticSourceControl/Source/PlasticSourceControl/Private/PackageUtils.cpp
@@ -0,0 +1,190 @@
+// Copyright Epic Games, Inc. All Rights Reserved.
+
+#include "PackageUtils.h"
+
+#include "Async/Async.h"
+#include "PackageTools.h"
+#include "Editor.h"
+#include "UObject/Linker.h"
+
+#include "ISourceControlModule.h" // LogSourceControl
+
+namespace PackageUtils
+{
+
+// Get the World currently loaded by the Editor (and thus, access to the corresponding map package)
+static UWorld* GetCurrentWorld()
+{
+	if (GEditor)
+	{
+		if (UWorld* EditorWorld = GEditor->GetEditorWorldContext().World())
+		{
+			return EditorWorld;
+		}
+	}
+	return nullptr;
+}
+
+// Find the packages corresponding to the files, if they are loaded in memory (won't load them)
+// Note: Extracted from AssetViewUtils::SyncPathsFromSourceControl()
+static TArray<UPackage*> FileNamesToLoadedPackages(const TArray<FString>& InFiles)
+{
+	TArray<UPackage*> LoadedPackages;
+	LoadedPackages.Reserve(InFiles.Num() + 1);
+	for (const FString& FilePath : InFiles)
+	{
+		FString PackageName;
+		FString FailureReason;
+		if (FPackageName::TryConvertFilenameToLongPackageName(FilePath, PackageName, &FailureReason))
+		{
+			// NOTE: this will only find packages loaded in memory
+			if (UPackage* Package = FindPackage(nullptr, *PackageName))
+			{
+				LoadedPackages.Emplace(Package);
+			}
+		}
+		// else, it means the file is not an asset from the Content/ folder (eg config, source code, anything else)
+	}
+	return LoadedPackages;
+}
+
+// Unlink all loaded packages to allow to update them
+// Note: Extracted from AssetViewUtils::SyncPathsFromSourceControl()
+void UnlinkPackages(const TArray<FString>& InFiles)
+{
+	const TArray<UPackage*> LoadedPackages = FileNamesToLoadedPackages(InFiles);
+	for (UPackage* Package : LoadedPackages)
+	{
+		// Detach the linkers of any loaded packages so that SCC can overwrite the files...
+		if (!Package->IsFullyLoaded())
+		{
+			FlushAsyncLoading();
+			Package->FullyLoad();
+		}
+		ResetLoaders(Package);
+	}
+	if (LoadedPackages.Num() > 0)
+	{
+		UE_LOG(LogSourceControl, Log, TEXT("Reseted Loader for %d Packages"), LoadedPackages.Num());
+	}
+}
+
+// Uses AsyncTask to call UnlinkPackages() on the Game Thread, and use a Promise to wait for the operation to complete
+void UnlinkPackagesInMainThread(const TArray<FString>& InFiles)
+{
+	const TSharedRef<TPromise<void>, ESPMode::ThreadSafe> Promise = MakeShareable(new TPromise<void>());
+	AsyncTask(ENamedThreads::GameThread, [Promise, InFiles]()
+	{
+		PackageUtils::UnlinkPackages(InFiles);
+		Promise->SetValue();
+	});
+	Promise->GetFuture().Get();
+}
+
+static TArray<UPackage*> ListPackagesToReload(const TArray<FString>& InFiles)
+{
+	TArray<UPackage*> LoadedPackages = FileNamesToLoadedPackages(InFiles);
+
+	// Detects if some packages to reload are part of the current map
+	// (ie assets within __ExternalActors__ or __ExternalObjects__ from the new One File Per Actor (OFPA) in UE5)
+	// in which case the current map need to be reloaded, so it needs to be added to the list of packages if not already there
+	// (then UPackageTools::ReloadPackages() will handle unloading the map at the start of the reload, avoiding some crash, and reloading it at the end)
+	if (UWorld* CurrentWorld = GetCurrentWorld())
+	{
+		UPackage* CurrentMapPackage = CurrentWorld->GetOutermost();
+
+		// If the current map file has been updated, it will be reloaded automatically, so no need for the following
+		const FString CurrentMapFileAbsolute = FPaths::ConvertRelativePathToFull(CurrentMapPackage->GetLoadedPath().GetLocalFullPath());
+		const bool bHasCurrentMapBeenUpdated = InFiles.FindByPredicate(
+			[&CurrentMapFileAbsolute](const FString& InFilePath) { return InFilePath.Equals(CurrentMapFileAbsolute, ESearchCase::IgnoreCase); }
+		) != nullptr;
+
+		if (!bHasCurrentMapBeenUpdated)
+		{
+			static const FString GamePath = FString("/Game");
+			const FString CurrentMapPath = *CurrentMapPackage->GetName();																	// eg "/Game/Maps/OpenWorld"
+			const FString CurrentMapPathWithoutGamePrefix = CurrentMapPath.RightChop(GamePath.Len());										// eg "/Maps/OpenWorld"
+			const FString CurrentMapExternalActorPath = FPackagePath::GetExternalActorsFolderName() + CurrentMapPathWithoutGamePrefix;		// eg "/__ExternalActors__/Maps/OpenWorld"
+			const FString CurrentMapExternalObjectPath = FPackagePath::GetExternalObjectsFolderName() + CurrentMapPathWithoutGamePrefix;	// eg "/__ExternalObjects__/Maps/OpenWorld"
+
+			bool bNeedReloadCurrentMap = false;
+
+			for (const FString& FilePath : InFiles)
+			{
+				if (FilePath.Contains(CurrentMapExternalActorPath) || FilePath.Contains(CurrentMapExternalObjectPath))
+				{
+					bNeedReloadCurrentMap = true;
+					break;
+				}
+			}
+
+			if (bNeedReloadCurrentMap)
+			{
+				LoadedPackages.Add(CurrentMapPackage);
+				UE_LOG(LogSourceControl, Log, TEXT("Reload: %s"), *CurrentMapPath);
+			}
+		}
+	}
+
+	return LoadedPackages;
+}
+
+// Hot-Reload all packages after they have been updated
+// Note: Extracted from AssetViewUtils::SyncPathsFromSourceControl()
+void ReloadPackages(TArray<UPackage*>& InPackages)
+{
+	UE_LOG(LogSourceControl, Log, TEXT("Reloading %d Packages..."), InPackages.Num());
+
+	// Syncing may have deleted some packages, so we need to unload those rather than re-load them...
+	// Note: we will store the package using weak pointers here otherwise we might have garbage collection issues after the ReloadPackages call
+	TArray<TWeakObjectPtr<UPackage>> PackagesToUnload;
+	InPackages.RemoveAll([&](UPackage* InPackage) -> bool
+	{
+		const FString PackageExtension = InPackage->ContainsMap() ? FPackageName::GetMapPackageExtension() : FPackageName::GetAssetPackageExtension();
+		const FString PackageFilename = FPackageName::LongPackageNameToFilename(InPackage->GetName(), PackageExtension);
+		if (!FPaths::FileExists(PackageFilename))
+		{
+			PackagesToUnload.Emplace(MakeWeakObjectPtr(InPackage));
+			return true; // remove package
+		}
+		return false; // keep package
+	});
+
+	// Hot-reload the new packages...
+	UPackageTools::ReloadPackages(InPackages);
+
+	// Unload any deleted packages...
+	TArray<UPackage*> PackageRawPtrsToUnload;
+	for (TWeakObjectPtr<UPackage>& PackageToUnload : PackagesToUnload)
+	{
+		if (PackageToUnload.IsValid())
+		{
+			PackageRawPtrsToUnload.Emplace(PackageToUnload.Get());
+		}
+	}
+
+	UPackageTools::UnloadPackages(PackageRawPtrsToUnload);
+}
+
+// Reload packages that where updated by the operation (and the current map if needed)
+void ReloadPackages(const TArray<FString>& InFiles)
+{
+	TArray<UPackage*> PackagesToReload = ListPackagesToReload(InFiles);
+	if (PackagesToReload.Num() > 0)
+	{
+		ReloadPackages(PackagesToReload);
+	}
+}
+
+void ReloadPackagesInMainThread(const TArray<FString>& InFiles)
+{
+	TSharedRef<TPromise<void>, ESPMode::ThreadSafe> Promise = MakeShareable(new TPromise<void>());
+	AsyncTask(ENamedThreads::GameThread, [Promise, InFiles]()
+	{
+		PackageUtils::ReloadPackages(InFiles);
+		Promise->SetValue();
+	});
+	Promise->GetFuture().Get();
+}
+
+} // namespace PackageUtils
diff --git a/Engine/Plugins/Developer/PlasticSourceControl/Source/PlasticSourceControl/Private/PackageUtils.h b/Engine/Plugins/Developer/PlasticSourceControl/Source/PlasticSourceControl/Private/PackageUtils.h
new file mode 100644
--- /dev/null
+++ b/Engine/Plugins/Developer/PlasticSourceControl/Source/PlasticSourceControl/Private/PackageUtils.h
@@ -0,0 +1,14 @@
+// Copyright Epic Games, Inc. All Rights Reserved.
+
+#pragma once
+
+#include "CoreMinimal.h"
+
+namespace PackageUtils
+{
+	void UnlinkPackages(const TArray<FString>& InFiles);
+	void UnlinkPackagesInMainThread(const TArray<FString>& InFiles);
+
+	void ReloadPackages(const TArray<FString>& InFiles);
+	void ReloadPackagesInMainThread(const TArray<FString>& InFiles);
+} // namespace PackageUtils
diff --git a/Engine/Plugins/Developer/PlasticSourceControl/Source/PlasticSourceControl/Private/PlasticSourceControlChangelist.h b/Engine/Plugins/Developer/PlasticSourceControl/Source/PlasticSourceControl/Private/PlasticSourceControlChangelist.h
--- a/Engine/Plugins/Developer/PlasticSourceControl/Source/PlasticSourceControl/Private/PlasticSourceControlChangelist.h
+++ b/Engine/Plugins/Developer/PlasticSourceControl/Source/PlasticSourceControl/Private/PlasticSourceControlChangelist.h
@@ -7,7 +7,7 @@
 #include "ISourceControlChangelist.h"
 
 /**
- * An abstraction of a changelist under source control: an identifier (a "name" in Plastic SCM)
+ * Unique Identifier of a changelist under source control: a "name" in Plastic SCM
  */
 class FPlasticSourceControlChangelist : public ISourceControlChangelist
 {
diff --git a/Engine/Plugins/Developer/PlasticSourceControl/Source/PlasticSourceControl/Private/PlasticSourceControlChangelistState.h b/Engine/Plugins/Developer/PlasticSourceControl/Source/PlasticSourceControl/Private/PlasticSourceControlChangelistState.h
--- a/Engine/Plugins/Developer/PlasticSourceControl/Source/PlasticSourceControl/Private/PlasticSourceControlChangelistState.h
+++ b/Engine/Plugins/Developer/PlasticSourceControl/Source/PlasticSourceControl/Private/PlasticSourceControlChangelistState.h
@@ -5,10 +5,11 @@
 #include "CoreMinimal.h"
 
 #include "ISourceControlChangelistState.h"
+
 #include "PlasticSourceControlChangelist.h"
 
 /**
- * An abstraction of the state of a pending changelist under source control: description and list of files
+ * The state of a pending changelist under source control: description and list of files
  */
 class FPlasticSourceControlChangelistState : public ISourceControlChangelistState
 {
@@ -74,7 +75,10 @@ public:
 
 	TArray<FSourceControlStateRef> Files;
 
-	TArray<FSourceControlStateRef> ShelvedFiles; // TODO: shelves are not yet implemented
+	int32 ShelveId = ISourceControlState::INVALID_REVISION;
+	FDateTime ShelveDate;
+
+	TArray<FSourceControlStateRef> ShelvedFiles;
 
 	/** The timestamp of the last update */
 	FDateTime TimeStamp;
diff --git a/Engine/Plugins/Developer/PlasticSourceControl/Source/PlasticSourceControl/Private/PlasticSourceControlConsole.cpp b/Engine/Plugins/Developer/PlasticSourceControl/Source/PlasticSourceControl/Private/PlasticSourceControlConsole.cpp
--- a/Engine/Plugins/Developer/PlasticSourceControl/Source/PlasticSourceControl/Private/PlasticSourceControlConsole.cpp
+++ b/Engine/Plugins/Developer/PlasticSourceControl/Source/PlasticSourceControl/Private/PlasticSourceControlConsole.cpp
@@ -38,6 +38,9 @@ void FPlasticSourceControlConsole::ExecutePlasticConsoleCommand(const TArray<FSt
 	const FString Command = a_args[0];
 	TArray<FString> Parameters = a_args;
 	Parameters.RemoveAt(0);
-	PlasticSourceControlUtils::RunCommandInternal(Command, Parameters, TArray<FString>(), EConcurrency::Synchronous, Results, Errors);
-	UE_LOG(LogSourceControl, Log, TEXT("Output:\n%s"), *Results);
+	PlasticSourceControlUtils::RunCommand(Command, Parameters, TArray<FString>(), Results, Errors);
+	if (Results.Len() > 200) // RunCommand() already log all command results up to 200 characters (limit to avoid long wall of text like XML)
+	{
+		UE_LOG(LogSourceControl, Log, TEXT("Output:\n%s"), *Results);
+	}
 }
diff --git a/Engine/Plugins/Developer/PlasticSourceControl/Source/PlasticSourceControl/Private/PlasticSourceControlMenu.cpp b/Engine/Plugins/Developer/PlasticSourceControl/Source/PlasticSourceControl/Private/PlasticSourceControlMenu.cpp
--- a/Engine/Plugins/Developer/PlasticSourceControl/Source/PlasticSourceControl/Private/PlasticSourceControlMenu.cpp
+++ b/Engine/Plugins/Developer/PlasticSourceControl/Source/PlasticSourceControl/Private/PlasticSourceControlMenu.cpp
@@ -15,11 +15,10 @@
 #include "UObject/Linker.h"
 #include "Widgets/Notifications/SNotificationList.h"
 #include "Framework/Notifications/NotificationManager.h"
-#include "Framework/MultiBox/MultiBoxBuilder.h"
 #include "Misc/MessageDialog.h"
 #include "Styling/AppStyle.h"
 
-#include "PackageTools.h"
+#include "PackageUtils.h"
 #include "FileHelpers.h"
 #include "ISettingsModule.h"
 
@@ -87,83 +86,9 @@ bool FPlasticSourceControlMenu::SaveDirtyPackages()
 /// Find all packages in Content directory
 TArray<FString> FPlasticSourceControlMenu::ListAllPackages()
 {
-	TArray<FString> PackageRelativePaths;
-	FPackageName::FindPackagesInDirectory(PackageRelativePaths, FPaths::ConvertRelativePathToFull(FPaths::ProjectContentDir()));
-
-	TArray<FString> PackageNames;
-	PackageNames.Reserve(PackageRelativePaths.Num());
-	for (const FString& Path : PackageRelativePaths)
-	{
-		FString PackageName;
-		FString FailureReason;
-		if (FPackageName::TryConvertFilenameToLongPackageName(Path, PackageName, &FailureReason))
-		{
-			PackageNames.Add(PackageName);
-		}
-		else
-		{
-			FMessageLog("SourceControl").Error(FText::FromString(FailureReason));
-		}
-	}
-
-	return PackageNames;
-}
-
-/// Unkink all loaded packages to allow to update them
-TArray<UPackage*> FPlasticSourceControlMenu::UnlinkPackages(const TArray<FString>& InPackageNames)
-{
-	TArray<UPackage*> LoadedPackages;
-
-	// Inspired from ContentBrowserUtils::SyncPathsFromSourceControl()
-	if (InPackageNames.Num() > 0)
-	{
-		// Form a list of loaded packages to reload...
-		LoadedPackages.Reserve(InPackageNames.Num());
-		for (const FString& PackageName : InPackageNames)
-		{
-			UPackage* Package = FindPackage(nullptr, *PackageName);
-			if (Package)
-			{
-				LoadedPackages.Emplace(Package);
-
-				// Detach the linkers of any loaded packages so that SCC can overwrite the files...
-				if (!Package->IsFullyLoaded())
-				{
-					FlushAsyncLoading();
-					Package->FullyLoad();
-				}
-				ResetLoaders(Package);
-			}
-		}
-		UE_LOG(LogSourceControl, Log, TEXT("Reseted Loader for %d Packages"), LoadedPackages.Num());
-	}
-
-	return LoadedPackages;
-}
-
-void FPlasticSourceControlMenu::ReloadPackages(TArray<UPackage*>& InPackagesToReload)
-{
-	UE_LOG(LogSourceControl, Log, TEXT("Reloading %d Packages..."), InPackagesToReload.Num());
-
-	// Syncing may have deleted some packages, so we need to unload those rather than re-load them...
-	TArray<UPackage*> PackagesToUnload;
-	InPackagesToReload.RemoveAll([&](UPackage* InPackage) -> bool
-	{
-		const FString PackageExtension = InPackage->ContainsMap() ? FPackageName::GetMapPackageExtension() : FPackageName::GetAssetPackageExtension();
-		const FString PackageFilename = FPackageName::LongPackageNameToFilename(InPackage->GetName(), PackageExtension);
-		if (!FPaths::FileExists(PackageFilename))
-		{
-			PackagesToUnload.Emplace(InPackage);
-			return true; // remove package
-		}
-		return false; // keep package
-	});
-
-	// Hot-reload the new packages...
-	UPackageTools::ReloadPackages(InPackagesToReload);
-
-	// Unload any deleted packages...
-	UPackageTools::UnloadPackages(PackagesToUnload);
+	TArray<FString> PackageFilePaths;
+	FPackageName::FindPackagesInDirectory(PackageFilePaths, FPaths::ConvertRelativePathToFull(FPaths::ProjectContentDir()));
+	return PackageFilePaths;
 }
 
 void FPlasticSourceControlMenu::SyncProjectClicked()
@@ -174,11 +99,11 @@ void FPlasticSourceControlMenu::SyncProjectClicked()
 		if (bSaved)
 		{
 			// Find and Unlink all packages in Content directory to allow to update them
-			PackagesToReload = UnlinkPackages(ListAllPackages());
+			PackageUtils::UnlinkPackages(ListAllPackages());
 
-			// Launch a "Sync" operation
+			// Launch a custom "SyncAll" operation
 			FPlasticSourceControlProvider& Provider = FPlasticSourceControlModule::Get().GetProvider();
-			TSharedRef<FSync, ESPMode::ThreadSafe> SyncOperation = ISourceControlOperation::Create<FSync>();
+			TSharedRef<FPlasticSyncAll, ESPMode::ThreadSafe> SyncOperation = ISourceControlOperation::Create<FPlasticSyncAll>();
 			const ECommandResult::Type Result = Provider.Execute(SyncOperation, TArray<FString>(), EConcurrency::Asynchronous, FSourceControlOperationComplete::CreateRaw(this, &FPlasticSourceControlMenu::OnSourceControlOperationComplete));
 			if (Result == ECommandResult::Succeeded)
 			{
@@ -187,9 +112,8 @@ void FPlasticSourceControlMenu::SyncProjectClicked()
 			}
 			else
 			{
-				// Report failure with a notification and Reload all packages
+				// Report failure with a notification (but nothing need to be reloaded since no local change is expected)
 				DisplayFailureNotification(SyncOperation->GetName());
-				ReloadPackages(PackagesToReload);
 			}
 		}
 		else
@@ -214,9 +138,7 @@ void FPlasticSourceControlMenu::RevertUnchangedClicked()
 		// Launch a "RevertUnchanged" Operation
 		FPlasticSourceControlProvider& Provider = FPlasticSourceControlModule::Get().GetProvider();
 		TSharedRef<FPlasticRevertUnchanged, ESPMode::ThreadSafe> RevertUnchangedOperation = ISourceControlOperation::Create<FPlasticRevertUnchanged>();
-		TArray<FString> WorkspaceRoot;
-		WorkspaceRoot.Add(Provider.GetPathToWorkspaceRoot()); // Revert the root of the workspace (needs an absolute path)
-		const ECommandResult::Type Result = Provider.Execute(RevertUnchangedOperation, WorkspaceRoot, EConcurrency::Asynchronous, FSourceControlOperationComplete::CreateRaw(this, &FPlasticSourceControlMenu::OnSourceControlOperationComplete));
+		const ECommandResult::Type Result = Provider.Execute(RevertUnchangedOperation, TArray<FString>(), EConcurrency::Asynchronous, FSourceControlOperationComplete::CreateRaw(this, &FPlasticSourceControlMenu::OnSourceControlOperationComplete));
 		if (Result == ECommandResult::Succeeded)
 		{
 			// Display an ongoing notification during the whole operation
@@ -249,14 +171,12 @@ void FPlasticSourceControlMenu::RevertAllClicked()
 			if (bSaved)
 			{
 				// Find and Unlink all packages in Content directory to allow to update them
-				PackagesToReload = UnlinkPackages(ListAllPackages());
+				PackageUtils::UnlinkPackages(ListAllPackages());
 
 				// Launch a "RevertAll" Operation
 				FPlasticSourceControlProvider& Provider = FPlasticSourceControlModule::Get().GetProvider();
 				TSharedRef<FPlasticRevertAll, ESPMode::ThreadSafe> RevertAllOperation = ISourceControlOperation::Create<FPlasticRevertAll>();
-				TArray<FString> WorkspaceRoot;
-				WorkspaceRoot.Add(Provider.GetPathToWorkspaceRoot()); // Revert the root of the workspace (needs an absolute path)
-				const ECommandResult::Type Result = Provider.Execute(RevertAllOperation, WorkspaceRoot, EConcurrency::Asynchronous, FSourceControlOperationComplete::CreateRaw(this, &FPlasticSourceControlMenu::OnSourceControlOperationComplete));
+				const ECommandResult::Type Result = Provider.Execute(RevertAllOperation, TArray<FString>(), EConcurrency::Asynchronous, FSourceControlOperationComplete::CreateRaw(this, &FPlasticSourceControlMenu::OnSourceControlOperationComplete));
 				if (Result == ECommandResult::Succeeded)
 				{
 					// Display an ongoing notification during the whole operation
@@ -264,9 +184,8 @@ void FPlasticSourceControlMenu::RevertAllClicked()
 				}
 				else
 				{
-					// Report failure with a notification and Reload all packages
+					// Report failure with a notification (but nothing need to be reloaded since no local change is expected)
 					DisplayFailureNotification(RevertAllOperation->GetName());
-					ReloadPackages(PackagesToReload);
 				}
 			}
 			else
@@ -416,10 +335,17 @@ void FPlasticSourceControlMenu::OnSourceControlOperationComplete(const FSourceCo
 {
 	RemoveInProgressNotification();
 
-	if ((InOperation->GetName() == "Sync") || (InOperation->GetName() == "RevertAll"))
+	if (InOperation->GetName() == "SyncAll")
+	{
+		// Reload packages that where updated by the Sync operation (and the current map if needed)
+		TSharedRef<FPlasticSyncAll, ESPMode::ThreadSafe> Operation = StaticCastSharedRef<FPlasticSyncAll>(InOperation);
+		PackageUtils::ReloadPackages(Operation->UpdatedFiles);
+	}
+	else if (InOperation->GetName() == "RevertAll")
 	{
-		// Reload packages that where unlinked at the beginning of the Sync operation
-		ReloadPackages(PackagesToReload);
+		// Reload packages that where updated by the Revert operation (and the current map if needed)
+		TSharedRef<FPlasticRevertAll, ESPMode::ThreadSafe> Operation = StaticCastSharedRef<FPlasticRevertAll>(InOperation);
+		PackageUtils::ReloadPackages(Operation->UpdatedFiles);
 	}
 
 	// Report result with a notification
@@ -435,6 +361,17 @@ void FPlasticSourceControlMenu::OnSourceControlOperationComplete(const FSourceCo
 
 void FPlasticSourceControlMenu::AddMenuExtension(FToolMenuSection& Menu)
 {
+	Menu.AddMenuEntry(
+		"PlasticSync",
+		LOCTEXT("PlasticSync",			"Sync/Update Workspace"),
+		LOCTEXT("PlasticSyncTooltip",	"Update the workspace to the latest changeset of the branch, and reload all affected assets."),
+		FSlateIcon(FAppStyle::GetAppStyleSetName(), "SourceControl.Actions.Sync"),
+		FUIAction(
+			FExecuteAction::CreateRaw(this, &FPlasticSourceControlMenu::SyncProjectClicked),
+			FCanExecuteAction()
+		)
+	);
+
 	Menu.AddMenuEntry(
 		"PlasticRevertUnchanged",
 		LOCTEXT("PlasticRevertUnchanged",			"Revert Unchanged"),
@@ -446,10 +383,21 @@ void FPlasticSourceControlMenu::AddMenuExtension(FToolMenuSection& Menu)
 		)
 	);
 
+	Menu.AddMenuEntry(
+		"PlasticRevertAll",
+		LOCTEXT("PlasticRevertAll",			"Revert All"),
+		LOCTEXT("PlasticRevertAllTooltip",	"Revert all files in the workspace to their controlled/unchanged state."),
+		FSlateIcon(FAppStyle::GetAppStyleSetName(), "SourceControl.Actions.Revert"),
+		FUIAction(
+			FExecuteAction::CreateRaw(this, &FPlasticSourceControlMenu::RevertAllClicked),
+			FCanExecuteAction()
+		)
+	);
+
 	Menu.AddMenuEntry(
 		"PlasticRefresh",
 		LOCTEXT("PlasticRefresh",			"Refresh"),
-		LOCTEXT("PlasticRefreshTooltip",	"Update the source control status of all files in the workspace."),
+		LOCTEXT("PlasticRefreshTooltip",	"Update the local source control status of all files in the workspace (no expensive checks for locks or changes on other branches)."),
 		FSlateIcon(FAppStyle::GetAppStyleSetName(), "SourceControl.Actions.Refresh"),
 		FUIAction(
 			FExecuteAction::CreateRaw(this, &FPlasticSourceControlMenu::RefreshClicked),
@@ -478,6 +426,7 @@ void FPlasticSourceControlMenu::AddMenuExtension(FToolMenuSection& Menu)
 			FCanExecuteAction()
 		)
 	);
+
 	Menu.AddMenuEntry(
 		"PlasticProjectSettings",
 		LOCTEXT("PlasticProjectSettings",			"Project Settings - Source Control - Plastic SCM"),
@@ -512,5 +461,4 @@ void FPlasticSourceControlMenu::AddMenuExtension(FToolMenuSection& Menu)
 	);
 }
 
-
 #undef LOCTEXT_NAMESPACE
diff --git a/Engine/Plugins/Developer/PlasticSourceControl/Source/PlasticSourceControl/Private/PlasticSourceControlMenu.h b/Engine/Plugins/Developer/PlasticSourceControl/Source/PlasticSourceControl/Private/PlasticSourceControlMenu.h
--- a/Engine/Plugins/Developer/PlasticSourceControl/Source/PlasticSourceControl/Private/PlasticSourceControlMenu.h
+++ b/Engine/Plugins/Developer/PlasticSourceControl/Source/PlasticSourceControl/Private/PlasticSourceControlMenu.h
@@ -13,7 +13,7 @@ class FPlasticSourceControlMenu
 public:
 	void Register();
 	void Unregister();
-	
+
 	/** This functions will be bound to appropriate Command. */
 	void SyncProjectClicked();
 	void RevertUnchangedClicked();
@@ -30,12 +30,8 @@ private:
 
 	bool				SaveDirtyPackages();
 	TArray<FString>		ListAllPackages();
-	TArray<UPackage*>	UnlinkPackages(const TArray<FString>& InPackageNames);
-	void				ReloadPackages(TArray<UPackage*>& InPackagesToReload);
-
-	void AddMenuExtension(FToolMenuSection& Builder);
 
-	TSharedRef<class FExtender> OnExtendLevelEditorViewMenu(const TSharedRef<class FUICommandList> CommandList);
+	void AddMenuExtension(FToolMenuSection& Menu);
 
 	void DisplayInProgressNotification(const FText& InOperationInProgressString);
 	void RemoveInProgressNotification();
@@ -43,10 +39,6 @@ private:
 	void DisplayFailureNotification(const FName& InOperationName);
 
 private:
-
-	/** Loaded packages to reload after a Sync or Revert operation */
-	TArray<UPackage*> PackagesToReload;
-
 	/** Current source control operation from extended menu if any */
 	TWeakPtr<class SNotificationItem> OperationInProgressNotification;
 
diff --git a/Engine/Plugins/Developer/PlasticSourceControl/Source/PlasticSourceControl/Private/PlasticSourceControlOperations.cpp b/Engine/Plugins/Developer/PlasticSourceControl/Source/PlasticSourceControl/Private/PlasticSourceControlOperations.cpp
--- a/Engine/Plugins/Developer/PlasticSourceControl/Source/PlasticSourceControl/Private/PlasticSourceControlOperations.cpp
+++ b/Engine/Plugins/Developer/PlasticSourceControl/Source/PlasticSourceControl/Private/PlasticSourceControlOperations.cpp
@@ -2,27 +2,28 @@
 
 #include "PlasticSourceControlOperations.h"
 
+#include "PackageUtils.h"
 #include "PlasticSourceControlCommand.h"
 #include "PlasticSourceControlModule.h"
 #include "PlasticSourceControlProvider.h"
 #include "PlasticSourceControlSettings.h"
 #include "PlasticSourceControlState.h"
 #include "PlasticSourceControlUtils.h"
+#include "PlasticSourceControlVersions.h"
+#include "ScopedTempFile.h"
 
+#include "Algo/AnyOf.h"
+#include "Algo/NoneOf.h"
 #include "AssetRegistry/AssetRegistryModule.h"
 #include "Async/Async.h"
 #include "HAL/FileManager.h"
 #include "Misc/Paths.h"
-#include "SourceControlOperations.h"
+
 #include "ISourceControlModule.h"
-#include "Algo/NoneOf.h"
+#include "SourceControlOperations.h"
 
 #define LOCTEXT_NAMESPACE "PlasticSourceControl"
 
-// 11.0.16.7248 add support for --descriptionfile for multi-line descriptions and support for special characters
-// https://www.plasticscm.com/download/releasenotes/11.0.16.7248
-static const FSoftwareVersion s_NewChangelistFileArgsPlasticScmVersion(11, 0, 16, 7248);
-
 template<typename Type>
 static FPlasticSourceControlWorkerRef InstantiateWorker(FPlasticSourceControlProvider& PlasticSourceControlProvider)
 {
@@ -42,6 +43,7 @@ void IPlasticSourceControlWorker::RegisterWorkers(FPlasticSourceControlProvider&
 	PlasticSourceControlProvider.RegisterWorker("RevertAll", FGetPlasticSourceControlWorker::CreateStatic(&InstantiateWorker<FPlasticRevertAllWorker>));
 	PlasticSourceControlProvider.RegisterWorker("MakeWorkspace", FGetPlasticSourceControlWorker::CreateStatic(&InstantiateWorker<FPlasticMakeWorkspaceWorker>));
 	PlasticSourceControlProvider.RegisterWorker("Sync", FGetPlasticSourceControlWorker::CreateStatic(&InstantiateWorker<FPlasticSyncWorker>));
+	PlasticSourceControlProvider.RegisterWorker("SyncAll", FGetPlasticSourceControlWorker::CreateStatic(&InstantiateWorker<FPlasticSyncWorker>));
 	PlasticSourceControlProvider.RegisterWorker("CheckIn", FGetPlasticSourceControlWorker::CreateStatic(&InstantiateWorker<FPlasticCheckInWorker>));
 	PlasticSourceControlProvider.RegisterWorker("Copy", FGetPlasticSourceControlWorker::CreateStatic(&InstantiateWorker<FPlasticCopyWorker>));
 	PlasticSourceControlProvider.RegisterWorker("Resolve", FGetPlasticSourceControlWorker::CreateStatic(&InstantiateWorker<FPlasticResolveWorker>));
@@ -51,6 +53,10 @@ void IPlasticSourceControlWorker::RegisterWorkers(FPlasticSourceControlProvider&
 	PlasticSourceControlProvider.RegisterWorker("DeleteChangelist", FGetPlasticSourceControlWorker::CreateStatic(&InstantiateWorker<FPlasticDeleteChangelistWorker>));
 	PlasticSourceControlProvider.RegisterWorker("EditChangelist", FGetPlasticSourceControlWorker::CreateStatic(&InstantiateWorker<FPlasticEditChangelistWorker>));
 	PlasticSourceControlProvider.RegisterWorker("MoveToChangelist", FGetPlasticSourceControlWorker::CreateStatic(&InstantiateWorker<FPlasticReopenWorker>));
+
+	PlasticSourceControlProvider.RegisterWorker("Shelve", FGetPlasticSourceControlWorker::CreateStatic(&InstantiateWorker<FPlasticShelveWorker>));
+	PlasticSourceControlProvider.RegisterWorker("Unshelve", FGetPlasticSourceControlWorker::CreateStatic(&InstantiateWorker<FPlasticUnshelveWorker>));
+	PlasticSourceControlProvider.RegisterWorker("DeleteShelved", FGetPlasticSourceControlWorker::CreateStatic(&InstantiateWorker<FPlasticDeleteShelveWorker>));
 }
 
 
@@ -110,7 +116,7 @@ bool FPlasticConnectWorker::Execute(FPlasticSourceControlCommand& InCommand)
 	if (GetProvider().IsPlasticAvailable())
 	{
 		// Get workspace name
-		InCommand.bCommandSuccessful = PlasticSourceControlUtils::GetWorkspaceName(GetProvider().GetPathToWorkspaceRoot(), InCommand.WorkspaceName, InCommand.ErrorMessages);
+		InCommand.bCommandSuccessful = PlasticSourceControlUtils::GetWorkspaceName(InCommand.PathToWorkspaceRoot, InCommand.WorkspaceName, InCommand.ErrorMessages);
 		if (InCommand.bCommandSuccessful)
 		{
 			// Get repository, server URL, branch and current changeset number
@@ -127,7 +133,7 @@ bool FPlasticConnectWorker::Execute(FPlasticSourceControlCommand& InCommand)
 				{
 					TArray<FString> ContentDir;
 					ContentDir.Add(FPaths::ConvertRelativePathToFull(FPaths::ProjectContentDir()));
-					PlasticSourceControlUtils::RunUpdateStatus(ContentDir, false, InCommand.Concurrency, InCommand.ErrorMessages, States, InCommand.ChangesetNumber, InCommand.BranchName);
+					PlasticSourceControlUtils::RunUpdateStatus(ContentDir, false, InCommand.ErrorMessages, States, InCommand.ChangesetNumber, InCommand.BranchName);
 				}
 			}
 			else
@@ -179,11 +185,20 @@ static void UpdateChangelistState(FPlasticSourceControlProvider& SCCProvider, co
 
 		for (const FPlasticSourceControlState& InState : InStates)
 		{
-			if ((InState.WorkspaceState != EWorkspaceState::CheckedOut) && (InState.WorkspaceState != EWorkspaceState::Added) && InState.WorkspaceState != EWorkspaceState::Deleted)
+			// cm cannot yet handle local modifications in changelists, only the GUI can
+			if (   InState.WorkspaceState != EWorkspaceState::CheckedOut
+				&& InState.WorkspaceState != EWorkspaceState::Added
+				&& InState.WorkspaceState != EWorkspaceState::Deleted
+				&& InState.WorkspaceState != EWorkspaceState::Copied
+				&& InState.WorkspaceState != EWorkspaceState::Moved
+				&& InState.WorkspaceState != EWorkspaceState::Conflicted	// In source control, waiting for merged
+				&& InState.WorkspaceState != EWorkspaceState::Replaced		// In source control, merged, waiting for checkin to conclude the merge
+				)
 			{
 				continue;
 			}
 
+			// Add a shared reference to the state of the file, that will then be updated by PlasticSourceControlUtils::UpdateCachedStates()
 			TSharedRef<FPlasticSourceControlState, ESPMode::ThreadSafe> State = SCCProvider.GetStateInternal(InState.GetFilename());
 			ChangelistState->Files.Add(State);
 
@@ -193,6 +208,7 @@ static void UpdateChangelistState(FPlasticSourceControlProvider& SCCProvider, co
 	}
 }
 
+
 FName FPlasticCheckOutWorker::GetName() const
 {
 	return "CheckOut";
@@ -207,15 +223,15 @@ bool FPlasticCheckOutWorker::Execute(FPlasticSourceControlCommand& InCommand)
 	// Detect special case for a partial checkout (CS:-1 in Gluon mode)!
 	if (-1 != InCommand.ChangesetNumber)
 	{
-		InCommand.bCommandSuccessful = PlasticSourceControlUtils::RunCommand(TEXT("checkout"), TArray<FString>(), InCommand.Files, InCommand.Concurrency, InCommand.InfoMessages, InCommand.ErrorMessages);
+		InCommand.bCommandSuccessful = PlasticSourceControlUtils::RunCommand(TEXT("checkout"), TArray<FString>(), InCommand.Files, InCommand.InfoMessages, InCommand.ErrorMessages);
 	}
 	else
 	{
-		InCommand.bCommandSuccessful = PlasticSourceControlUtils::RunCommand(TEXT("partial checkout"), TArray<FString>(), InCommand.Files, InCommand.Concurrency, InCommand.InfoMessages, InCommand.ErrorMessages);
+		InCommand.bCommandSuccessful = PlasticSourceControlUtils::RunCommand(TEXT("partial checkout"), TArray<FString>(), InCommand.Files, InCommand.InfoMessages, InCommand.ErrorMessages);
 	}
 
 	// now update the status of our files
-	PlasticSourceControlUtils::RunUpdateStatus(InCommand.Files, false, InCommand.Concurrency, InCommand.ErrorMessages, States, InCommand.ChangesetNumber, InCommand.BranchName);
+	PlasticSourceControlUtils::RunUpdateStatus(InCommand.Files, false, InCommand.ErrorMessages, States, InCommand.ChangesetNumber, InCommand.BranchName);
 
 	return InCommand.bCommandSuccessful;
 }
@@ -230,15 +246,15 @@ bool FPlasticCheckOutWorker::UpdateStates()
 	return PlasticSourceControlUtils::UpdateCachedStates(MoveTemp(States));
 }
 
-bool DeleteChangelist(const FPlasticSourceControlProvider& PlasticSourceControlProvider, const FPlasticSourceControlChangelist& InChangelist, const EConcurrency::Type InConcurrency, TArray<FString>& OutResults, TArray<FString>& OutErrorMessages)
+bool DeleteChangelist(const FPlasticSourceControlProvider& PlasticSourceControlProvider, const FPlasticSourceControlChangelist& InChangelist, TArray<FString>& OutResults, TArray<FString>& OutErrorMessages)
 {
 	TArray<FString> Parameters;
 	TArray<FString> Files;
-	if (PlasticSourceControlProvider.GetPlasticScmVersion() < s_NewChangelistFileArgsPlasticScmVersion)
+	if (PlasticSourceControlProvider.GetPlasticScmVersion() < PlasticSourceControlVersions::NewChangelistFileArgs)
 	{
 		Parameters.Add(TEXT("rm"));
 		Files.Add(InChangelist.GetName());
-		return PlasticSourceControlUtils::RunCommand(TEXT("changelist"), Parameters, Files, InConcurrency, OutResults, OutErrorMessages);
+		return PlasticSourceControlUtils::RunCommand(TEXT("changelist"), Parameters, Files, OutResults, OutErrorMessages);
 	}
 	else
 	{
@@ -246,7 +262,7 @@ bool DeleteChangelist(const FPlasticSourceControlProvider& PlasticSourceControlP
 		const FScopedTempFile ChangelistNameFile(InChangelist.GetName());
 		Parameters.Add(FString::Printf(TEXT("--namefile=\"%s\""), *FPaths::ConvertRelativePathToFull(ChangelistNameFile.GetFilename())));
 		UE_LOG(LogSourceControl, Verbose, TEXT("DeleteChangelist(%s)"), *InChangelist.GetName());
-		return PlasticSourceControlUtils::RunCommand(TEXT("changelist"), Parameters, Files, InConcurrency, OutResults, OutErrorMessages);
+		return PlasticSourceControlUtils::RunCommand(TEXT("changelist"), Parameters, Files, OutResults, OutErrorMessages);
 	}
 }
 
@@ -287,6 +303,23 @@ static FText ParseCheckInResults(const TArray<FString>& InResults)
 	return FText();
 }
 
+
+TArray<FString> GetFilesFromCommand(FPlasticSourceControlProvider& PlasticSourceControlProvider, FPlasticSourceControlCommand& InCommand)
+{
+	TArray<FString> Files;
+	if (InCommand.Changelist.IsInitialized() && InCommand.Files.IsEmpty())
+	{
+		TSharedRef<FPlasticSourceControlChangelistState, ESPMode::ThreadSafe> ChangelistState = PlasticSourceControlProvider.GetStateInternal(InCommand.Changelist);
+		Files = FileNamesFromFileStates(ChangelistState->Files);
+	}
+	else
+	{
+		Files = InCommand.Files;
+	}
+	return Files;
+}
+
+
 FName FPlasticCheckInWorker::GetName() const
 {
 	return "CheckIn";
@@ -298,26 +331,28 @@ bool FPlasticCheckInWorker::Execute(FPlasticSourceControlCommand& InCommand)
 
 	check(InCommand.Operation->GetName() == GetName());
 	TSharedRef<FCheckIn, ESPMode::ThreadSafe> Operation = StaticCastSharedRef<FCheckIn>(InCommand.Operation);
+	FText Description = Operation->GetDescription();
 
-	TArray<FString> Files;
-	if (InCommand.Changelist.IsInitialized() && InCommand.Files.IsEmpty())
-	{
-		TSharedRef<FPlasticSourceControlChangelistState, ESPMode::ThreadSafe> ChangelistState = GetProvider().GetStateInternal(InCommand.Changelist);
-		Files = FileNamesFromFileStates(ChangelistState->Files);
-
-		InChangelist = InCommand.Changelist;
-	}
-	else
-	{
-		Files = InCommand.Files;
-	}
+	const TArray<FString> Files = GetFilesFromCommand(GetProvider(), InCommand);
 
 	if (Files.Num() > 0)
 	{
-		UE_LOG(LogSourceControl, Verbose, TEXT("CheckIn: %d file(s) Description: '%s'"), Files.Num(), *Operation->GetDescription().ToString());
+		if (InCommand.Changelist.IsInitialized())
+		{
+			TSharedRef<FPlasticSourceControlChangelistState, ESPMode::ThreadSafe> ChangelistState = GetProvider().GetStateInternal(InCommand.Changelist);
+
+			if (Description.IsEmpty())
+			{
+				Description = ChangelistState->GetDescriptionText();
+			}
+
+			InChangelist = InCommand.Changelist;
+		}
+
+		UE_LOG(LogSourceControl, Verbose, TEXT("CheckIn: %d file(s) Description: '%s'"), Files.Num(), *Description.ToString());
 
 		// make a temp file to place our commit message in
-		const FScopedTempFile CommitMsgFile(Operation->GetDescription());
+		const FScopedTempFile CommitMsgFile(Description);
 		if (!CommitMsgFile.GetFilename().IsEmpty())
 		{
 			TArray<FString> Parameters;
@@ -328,12 +363,12 @@ bool FPlasticCheckInWorker::Execute(FPlasticSourceControlCommand& InCommand)
 				Parameters.Add(TEXT("--all")); // Also files Changed (not CheckedOut) and Moved/Deleted Locally
 			//  NOTE: --update added as #23 but removed as #32 because most assets are locked by the Unreal Editor
 			//  Parameters.Add(TEXT("--update")); // Processes the update-merge automatically if it eventually happens.
-				InCommand.bCommandSuccessful = PlasticSourceControlUtils::RunCommand(TEXT("checkin"), Parameters, Files, InCommand.Concurrency, InCommand.InfoMessages, InCommand.ErrorMessages);
+				InCommand.bCommandSuccessful = PlasticSourceControlUtils::RunCommand(TEXT("checkin"), Parameters, Files, InCommand.InfoMessages, InCommand.ErrorMessages);
 			}
 			else
 			{
 				Parameters.Add(TEXT("--applychanged")); // Also files Changed (not CheckedOut) and Moved/Deleted Locally
-				InCommand.bCommandSuccessful = PlasticSourceControlUtils::RunCommand(TEXT("partial checkin"), Parameters, Files, InCommand.Concurrency, InCommand.InfoMessages, InCommand.ErrorMessages);
+				InCommand.bCommandSuccessful = PlasticSourceControlUtils::RunCommand(TEXT("partial checkin"), Parameters, Files, InCommand.InfoMessages, InCommand.ErrorMessages);
 			}
 			if (InCommand.bCommandSuccessful)
 			{
@@ -354,13 +389,13 @@ bool FPlasticCheckInWorker::Execute(FPlasticSourceControlCommand& InCommand)
 
 			if (InChangelist.IsInitialized() && !InChangelist.IsDefault())
 			{
-				// NOTE: we need to explicitly delete persistent changelists when we submit its content
-				DeleteChangelist(GetProvider(), InChangelist, InCommand.Concurrency, InCommand.InfoMessages, InCommand.ErrorMessages);
+				// NOTE: we need to explicitly delete persistent changelists when we submit its content, except for the Default changelist
+				DeleteChangelist(GetProvider(), InChangelist, InCommand.InfoMessages, InCommand.ErrorMessages);
 			}
 		}
 
 		// now update the status of our files
-		PlasticSourceControlUtils::RunUpdateStatus(Files, false, InCommand.Concurrency, InCommand.ErrorMessages, States, InCommand.ChangesetNumber, InCommand.BranchName);
+		PlasticSourceControlUtils::RunUpdateStatus(Files, false, InCommand.ErrorMessages, States, InCommand.ChangesetNumber, InCommand.BranchName);
 	}
 	else
 	{
@@ -374,6 +409,7 @@ bool FPlasticCheckInWorker::UpdateStates()
 {
 	TRACE_CPUPROFILER_EVENT_SCOPE(FPlasticCheckInWorker::UpdateStates);
 
+	// Is the submit of a full changelist (from the View Changelists window) or a set of files (from the Submit Content window)?
 	if (InChangelist.IsInitialized())
 	{
 		if (InChangelist.IsDefault())
@@ -384,9 +420,32 @@ bool FPlasticCheckInWorker::UpdateStates()
 		}
 		else
 		{
+			// Remove references to the changelist (else it later creates a phantom of the old changelist in the cache)
+			for (const FPlasticSourceControlState& NewState : States)
+			{
+				TSharedRef<FPlasticSourceControlState, ESPMode::ThreadSafe> State = GetProvider().GetStateInternal(NewState.GetFilename());
+				State->Changelist.Reset();
+			}
+
 			GetProvider().RemoveChangelistFromCache(InChangelist);
 		}
 	}
+	else
+	{
+		// Update affected changelists if any
+		for (const FPlasticSourceControlState& NewState : States)
+		{
+			TSharedRef<FPlasticSourceControlState, ESPMode::ThreadSafe> State = GetProvider().GetStateInternal(NewState.GetFilename());
+			if (State->Changelist.IsInitialized())
+			{
+				// 1- Remove these files from their previous changelist
+				TSharedRef<FPlasticSourceControlChangelistState, ESPMode::ThreadSafe> PreviousChangelist = GetProvider().GetStateInternal(State->Changelist);
+				PreviousChangelist->Files.Remove(State);
+				// 2- And reset the reference to their previous changelist
+				State->Changelist.Reset();
+			}
+		}
+	}
 
 	return PlasticSourceControlUtils::UpdateCachedStates(MoveTemp(States));
 }
@@ -423,15 +482,15 @@ bool FPlasticMarkForAddWorker::Execute(FPlasticSourceControlCommand& InCommand)
 		// Detect special case for a partial checkout (CS:-1 in Gluon mode)!
 		if (-1 != InCommand.ChangesetNumber)
 		{
-			InCommand.bCommandSuccessful = PlasticSourceControlUtils::RunCommand(TEXT("add"), Parameters, InCommand.Files, InCommand.Concurrency, InCommand.InfoMessages, InCommand.ErrorMessages);
+			InCommand.bCommandSuccessful = PlasticSourceControlUtils::RunCommand(TEXT("add"), Parameters, InCommand.Files, InCommand.InfoMessages, InCommand.ErrorMessages);
 		}
 		else
 		{
-			InCommand.bCommandSuccessful = PlasticSourceControlUtils::RunCommand(TEXT("partial add"), Parameters, InCommand.Files, InCommand.Concurrency, InCommand.InfoMessages, InCommand.ErrorMessages);
+			InCommand.bCommandSuccessful = PlasticSourceControlUtils::RunCommand(TEXT("partial add"), Parameters, InCommand.Files, InCommand.InfoMessages, InCommand.ErrorMessages);
 		}
 
 		// now update the status of our files
-		PlasticSourceControlUtils::RunUpdateStatus(InCommand.Files, false, InCommand.Concurrency, InCommand.ErrorMessages, States, InCommand.ChangesetNumber, InCommand.BranchName);
+		PlasticSourceControlUtils::RunUpdateStatus(InCommand.Files, false, InCommand.ErrorMessages, States, InCommand.ChangesetNumber, InCommand.BranchName);
 	}
 	else
 	{
@@ -465,15 +524,15 @@ bool FPlasticDeleteWorker::Execute(FPlasticSourceControlCommand& InCommand)
 	// Detect special case for a partial checkout (CS:-1 in Gluon mode)!
 	if (-1 != InCommand.ChangesetNumber)
 	{
-		InCommand.bCommandSuccessful = PlasticSourceControlUtils::RunCommand(TEXT("remove"), TArray<FString>(), InCommand.Files, InCommand.Concurrency, InCommand.InfoMessages, InCommand.ErrorMessages);
+		InCommand.bCommandSuccessful = PlasticSourceControlUtils::RunCommand(TEXT("remove"), TArray<FString>(), InCommand.Files, InCommand.InfoMessages, InCommand.ErrorMessages);
 	}
 	else
 	{
-		InCommand.bCommandSuccessful = PlasticSourceControlUtils::RunCommand(TEXT("partial remove"), TArray<FString>(), InCommand.Files, InCommand.Concurrency, InCommand.InfoMessages, InCommand.ErrorMessages);
+		InCommand.bCommandSuccessful = PlasticSourceControlUtils::RunCommand(TEXT("partial remove"), TArray<FString>(), InCommand.Files, InCommand.InfoMessages, InCommand.ErrorMessages);
 	}
 
 	// now update the status of our files
-	PlasticSourceControlUtils::RunUpdateStatus(InCommand.Files, false, InCommand.Concurrency, InCommand.ErrorMessages, States, InCommand.ChangesetNumber, InCommand.BranchName);
+	PlasticSourceControlUtils::RunUpdateStatus(InCommand.Files, false, InCommand.ErrorMessages, States, InCommand.ChangesetNumber, InCommand.BranchName);
 
 	return InCommand.bCommandSuccessful;
 }
@@ -498,67 +557,55 @@ bool FPlasticRevertWorker::Execute(FPlasticSourceControlCommand& InCommand)
 	TRACE_CPUPROFILER_EVENT_SCOPE(FPlasticRevertWorker::Execute);
 
 	check(InCommand.Operation->GetName() == GetName());
+	TSharedRef<FRevert, ESPMode::ThreadSafe> Operation = StaticCastSharedRef<FRevert>(InCommand.Operation);
 
-	TArray<FString> Files;
-	if (InCommand.Changelist.IsInitialized() && InCommand.Files.IsEmpty())
-	{
-		TSharedRef<FPlasticSourceControlChangelistState, ESPMode::ThreadSafe> ChangelistState = GetProvider().GetStateInternal(InCommand.Changelist);
-		Files = FileNamesFromFileStates(ChangelistState->Files);
-	}
-	else
-	{
-		Files = InCommand.Files;
-	}
-
-	TArray<FString> ChangedFiles;
-	TArray<FString> CheckedOutFiles;
+	TArray<FString> Files = GetFilesFromCommand(GetProvider(), InCommand);
 
-	for (const FString& File : Files)
+	for (int i = 0; i < Files.Num(); i++) // Required for loop on index since we are adding to the Files array as we go
 	{
-		TSharedRef<FPlasticSourceControlState, ESPMode::ThreadSafe> State = GetProvider().GetStateInternal(File);
+		const FString& File = Files[i];
 
-		if (EWorkspaceState::Changed == State->WorkspaceState)
-		{
-			// only revert the changes of the given file in workspace
-			ChangedFiles.Add(State->LocalFilename);
-		}
-		else
+		const TSharedRef<const FPlasticSourceControlState, ESPMode::ThreadSafe> State = GetProvider().GetStateInternal(File);
+
+		if (State->WorkspaceState == EWorkspaceState::Moved)
 		{
-			CheckedOutFiles.Add(State->LocalFilename);
-			// in case of a Moved/Renamed, find the rename origin to revert both at once
-			if (EWorkspaceState::Moved == State->WorkspaceState)
-			{
-				CheckedOutFiles.Add(State->MovedFrom);
+			const FString& MovedFrom = State->MovedFrom;
 
-				// Delete the redirector
-				IFileManager::Get().Delete(*State->MovedFrom);
+			// In case of a file Moved/Renamed, consider the rename Origin (where there is now a Redirector file Added)
+			// and add it to the list of files to revert (only if it is not already in) to revert both at once
+			if (!Files.FindByPredicate([&MovedFrom](const FString& File) { return File.Equals(MovedFrom, ESearchCase::IgnoreCase); }))
+			{
+				Files.Add(MovedFrom);
 			}
+
+			// and delete the Redirector (else the reverted file will collide with it and create a *.private.0 file)
+			IFileManager::Get().Delete(*MovedFrom);
+		}
+
+		if (State->WorkspaceState == EWorkspaceState::Added && Operation->ShouldDeleteNewFiles())
+		{
+			IFileManager::Get().Delete(*File);
 		}
 	}
 
 	InCommand.bCommandSuccessful = true;
 
-	if (ChangedFiles.Num() > 0)
-	{
-		InCommand.bCommandSuccessful &= PlasticSourceControlUtils::RunCommand(TEXT("undochange"), TArray<FString>(), ChangedFiles, InCommand.Concurrency, InCommand.InfoMessages, InCommand.ErrorMessages);
-	}
-
-	if (CheckedOutFiles.Num() > 0)
+	if (Files.Num() > 0)
 	{
 		// revert the checkout and any changes of the given file in workspace
 		// Detect special case for a partial checkout (CS:-1 in Gluon mode)!
 		if (-1 != InCommand.ChangesetNumber)
 		{
-			InCommand.bCommandSuccessful &= PlasticSourceControlUtils::RunCommand(TEXT("undocheckout"), TArray<FString>(), CheckedOutFiles, InCommand.Concurrency, InCommand.InfoMessages, InCommand.ErrorMessages);
+			InCommand.bCommandSuccessful &= PlasticSourceControlUtils::RunCommand(TEXT("undo"), TArray<FString>(), Files, InCommand.InfoMessages, InCommand.ErrorMessages);
 		}
 		else
 		{
-			InCommand.bCommandSuccessful &= PlasticSourceControlUtils::RunCommand(TEXT("partial undocheckout"), TArray<FString>(), CheckedOutFiles, InCommand.Concurrency, InCommand.InfoMessages, InCommand.ErrorMessages);
+			InCommand.bCommandSuccessful &= PlasticSourceControlUtils::RunCommand(TEXT("partial undo"), TArray<FString>(), Files, InCommand.InfoMessages, InCommand.ErrorMessages);
 		}
 	}
 
 	// update the status of our files
-	PlasticSourceControlUtils::RunUpdateStatus(Files, false, InCommand.Concurrency, InCommand.ErrorMessages, States, InCommand.ChangesetNumber, InCommand.BranchName);
+	PlasticSourceControlUtils::RunUpdateStatus(Files, false, InCommand.ErrorMessages, States, InCommand.ChangesetNumber, InCommand.BranchName);
 
 	return InCommand.bCommandSuccessful;
 }
@@ -567,7 +614,7 @@ bool FPlasticRevertWorker::UpdateStates()
 {
 	TRACE_CPUPROFILER_EVENT_SCOPE(FPlasticRevertWorker::UpdateStates);
 
-	// Update affected changelist if any
+	// Update affected changelists if any
 	for (const FPlasticSourceControlState& NewState : States)
 	{
 		TSharedRef<FPlasticSourceControlState, ESPMode::ThreadSafe> State = GetProvider().GetStateInternal(NewState.GetFilename());
@@ -576,6 +623,8 @@ bool FPlasticRevertWorker::UpdateStates()
 			// 1- Remove these files from their previous changelist
 			TSharedRef<FPlasticSourceControlChangelistState, ESPMode::ThreadSafe> PreviousChangelist = GetProvider().GetStateInternal(State->Changelist);
 			PreviousChangelist->Files.Remove(State);
+			// 2- And reset the reference to their previous changelist
+			State->Changelist.Reset();
 		}
 	}
 
@@ -596,26 +645,17 @@ bool FPlasticRevertUnchangedWorker::Execute(FPlasticSourceControlCommand& InComm
 	TArray<FString> Parameters;
 	Parameters.Add(TEXT("-R"));
 
-	TArray<FString> Files;
-	if (InCommand.Changelist.IsInitialized() && InCommand.Files.IsEmpty())
-	{
-		TSharedRef<FPlasticSourceControlChangelistState, ESPMode::ThreadSafe> ChangelistState = GetProvider().GetStateInternal(InCommand.Changelist);
-		Files = FileNamesFromFileStates(ChangelistState->Files);
-	}
-	else
-	{
-		Files = InCommand.Files;
-	}
+	TArray<FString> Files = GetFilesFromCommand(GetProvider(), InCommand);
 
 	// revert the checkout of all unchanged files recursively
-	InCommand.bCommandSuccessful = PlasticSourceControlUtils::RunCommand(TEXT("uncounchanged"), Parameters, Files, InCommand.Concurrency, InCommand.InfoMessages, InCommand.ErrorMessages);
+	InCommand.bCommandSuccessful = PlasticSourceControlUtils::RunCommand(TEXT("uncounchanged"), Parameters, Files, InCommand.InfoMessages, InCommand.ErrorMessages);
 
 	// Now update the status of either the files, or all assets in the Content directory
 	if (Files.Num() == 0)
 	{
 		Files.Add(FPaths::ConvertRelativePathToFull(FPaths::ProjectContentDir()));
 	}
-	PlasticSourceControlUtils::RunUpdateStatus(Files, false, InCommand.Concurrency, InCommand.ErrorMessages, States, InCommand.ChangesetNumber, InCommand.BranchName);
+	PlasticSourceControlUtils::RunUpdateStatus(Files, false, InCommand.ErrorMessages, States, InCommand.ChangesetNumber, InCommand.BranchName);
 
 	return InCommand.bCommandSuccessful;
 }
@@ -624,7 +664,7 @@ bool FPlasticRevertUnchangedWorker::UpdateStates()
 {
 	TRACE_CPUPROFILER_EVENT_SCOPE(FPlasticRevertUnchangedWorker::UpdateStates);
 
-	// Update affected changelist if any
+	// Update affected changelists if any
 	for (const FPlasticSourceControlState& NewState : States)
 	{
 		if (!NewState.IsModified())
@@ -635,6 +675,8 @@ bool FPlasticRevertUnchangedWorker::UpdateStates()
 				// 1- Remove these files from their previous changelist
 				TSharedRef<FPlasticSourceControlChangelistState, ESPMode::ThreadSafe> PreviousChangelist = GetProvider().GetStateInternal(State->Changelist);
 				PreviousChangelist->Files.Remove(State);
+				// 2- And reset the reference to their previous changelist
+				State->Changelist.Reset();
 			}
 		}
 	}
@@ -652,24 +694,63 @@ bool FPlasticRevertAllWorker::Execute(FPlasticSourceControlCommand& InCommand)
 	TRACE_CPUPROFILER_EVENT_SCOPE(FPlasticRevertAllWorker::Execute);
 
 	check(InCommand.Operation->GetName() == GetName());
+	TSharedRef<FPlasticRevertAll, ESPMode::ThreadSafe> Operation = StaticCastSharedRef<FPlasticRevertAll>(InCommand.Operation);
 
+	// Start by updating the Status of all Content, to find all the changes that will be reverted
+	{
+		TArray<FPlasticSourceControlState> TempStates;
+		TArray<FString> ContentDir;
+		ContentDir.Add(FPaths::ConvertRelativePathToFull(FPaths::ProjectContentDir()));
+		PlasticSourceControlUtils::RunUpdateStatus(ContentDir, false, InCommand.ErrorMessages, TempStates, InCommand.ChangesetNumber, InCommand.BranchName);
+
+		for (auto& State : TempStates)
+		{
+			if (State.IsModified())
+			{
+				if (State.WorkspaceState == EWorkspaceState::Added && Operation->ShouldDeleteNewFiles())
+				{
+					IFileManager::Get().Delete(*State.GetFilename());
+				}
+
+				// Add all modified files to the list of files to be updated (reverted and then reloaded)
+				Operation->UpdatedFiles.Add(MoveTemp(State.LocalFilename));
+
+				if (State.WorkspaceState == EWorkspaceState::Moved)
+				{
+					// In case of a file Moved/Renamed, consider the rename Origin (where there is now a Redirector file Added)
+					// and add it to the list of files to revert
+					Operation->UpdatedFiles.Add(MoveTemp(State.MovedFrom));
+
+					// and delete the Redirector (else the reverted file will collide with it and create a *.private.0 file)
+					IFileManager::Get().Delete(*State.MovedFrom);
+				}
+			}
+		}
+	}
+
+	TArray<FString> Results;
 	TArray<FString> Parameters;
 	Parameters.Add(TEXT("--all"));
 	// revert the checkout of all files recursively
 	// Detect special case for a partial checkout (CS:-1 in Gluon mode)!
 	if (-1 != InCommand.ChangesetNumber)
 	{
-		InCommand.bCommandSuccessful = PlasticSourceControlUtils::RunCommand(TEXT("undocheckout"), Parameters, InCommand.Files, InCommand.Concurrency, InCommand.InfoMessages, InCommand.ErrorMessages);
+		InCommand.bCommandSuccessful = PlasticSourceControlUtils::RunCommand(TEXT("undocheckout"), Parameters, TArray<FString>(), Results, InCommand.ErrorMessages);
 	}
 	else
 	{
-		InCommand.bCommandSuccessful = PlasticSourceControlUtils::RunCommand(TEXT("partial undocheckout"), Parameters, InCommand.Files, InCommand.Concurrency, InCommand.InfoMessages, InCommand.ErrorMessages);
+		InCommand.bCommandSuccessful = PlasticSourceControlUtils::RunCommand(TEXT("partial undocheckout"), Parameters, TArray<FString>(), Results, InCommand.ErrorMessages);
 	}
 
-	// Now update the status of assets in the Content directory
-	TArray<FString> ContentDir;
-	ContentDir.Add(FPaths::ConvertRelativePathToFull(FPaths::ProjectContentDir()));
-	PlasticSourceControlUtils::RunUpdateStatus(ContentDir, false, InCommand.Concurrency, InCommand.ErrorMessages, States, InCommand.ChangesetNumber, InCommand.BranchName);
+	// NOTE: don't parse the Results, it has too many quirks, uses the list from the status update;
+	// - Renames are not easy to parse without a clean separator in "Origin Destination"
+	// - Files added in folders are not accounted for by "undo", only the folder is listed in the results
+
+	// now update the status of the updated files
+	if (Operation->UpdatedFiles.Num())
+	{
+		PlasticSourceControlUtils::RunUpdateStatus(Operation->UpdatedFiles, false, InCommand.ErrorMessages, States, InCommand.ChangesetNumber, InCommand.BranchName);
+	}
 
 	return InCommand.bCommandSuccessful;
 }
@@ -678,6 +759,23 @@ bool FPlasticRevertAllWorker::UpdateStates()
 {
 	TRACE_CPUPROFILER_EVENT_SCOPE(FPlasticRevertAllWorker::UpdateStates);
 
+	// Update affected changelists if any
+	for (const FPlasticSourceControlState& NewState : States)
+	{
+		if (!NewState.IsModified())
+		{
+			TSharedRef<FPlasticSourceControlState, ESPMode::ThreadSafe> State = GetProvider().GetStateInternal(NewState.GetFilename());
+			if (State->Changelist.IsInitialized())
+			{
+				// 1- Remove these files from their previous changelist
+				TSharedRef<FPlasticSourceControlChangelistState, ESPMode::ThreadSafe> PreviousChangelist = GetProvider().GetStateInternal(State->Changelist);
+				PreviousChangelist->Files.Remove(State);
+				// 2- And reset the reference to their previous changelist
+				State->Changelist.Reset();
+			}
+		}
+	}
+
 	return PlasticSourceControlUtils::UpdateCachedStates(MoveTemp(States));
 }
 
@@ -697,14 +795,14 @@ bool FPlasticMakeWorkspaceWorker::Execute(FPlasticSourceControlCommand& InComman
 		TArray<FString> Parameters;
 		Parameters.Add(Operation->ServerUrl);
 		Parameters.Add(Operation->RepositoryName);
-		PlasticSourceControlUtils::RunCommand(TEXT("makerepository"), Parameters, TArray<FString>(), EConcurrency::Synchronous, InCommand.InfoMessages, InCommand.ErrorMessages);
+		PlasticSourceControlUtils::RunCommand(TEXT("makerepository"), Parameters, TArray<FString>(), InCommand.InfoMessages, InCommand.ErrorMessages);
 	}
 	{
 		TArray<FString> Parameters;
 		Parameters.Add(Operation->WorkspaceName);
 		Parameters.Add(TEXT(".")); // current path, ie. ProjectDir
 		Parameters.Add(FString::Printf(TEXT("--repository=rep:%s@repserver:%s"), *Operation->RepositoryName, *Operation->ServerUrl));
-		InCommand.bCommandSuccessful = PlasticSourceControlUtils::RunCommand(TEXT("makeworkspace"), Parameters, TArray<FString>(), EConcurrency::Synchronous, InCommand.InfoMessages, InCommand.ErrorMessages);
+		InCommand.bCommandSuccessful = PlasticSourceControlUtils::RunCommand(TEXT("makeworkspace"), Parameters, TArray<FString>(), InCommand.InfoMessages, InCommand.ErrorMessages);
 	}
 
 	return InCommand.bCommandSuccessful;
@@ -731,27 +829,25 @@ bool FPlasticUpdateStatusWorker::Execute(FPlasticSourceControlCommand& InCommand
 	UE_LOG(LogSourceControl, Log, TEXT("status (of %d files, ShouldUpdateHistory=%d, ShouldGetOpenedOnly=%d, ShouldUpdateModifiedState=%d)"),
 		InCommand.Files.Num(), Operation->ShouldUpdateHistory(), Operation->ShouldGetOpenedOnly(), Operation->ShouldUpdateModifiedState());
 
-	TArray<FString> Files;
-	if (InCommand.Changelist.IsInitialized() && InCommand.Files.IsEmpty())
-	{
-		TSharedRef<FPlasticSourceControlChangelistState, ESPMode::ThreadSafe> ChangelistState = GetProvider().GetStateInternal(InCommand.Changelist);
-		Files = FileNamesFromFileStates(ChangelistState->Files);
-	}
-	else
-	{
-		Files = InCommand.Files;
-	}
+	const TArray<FString> Files = GetFilesFromCommand(GetProvider(), InCommand);
 
 	if (Files.Num() > 0)
 	{
-		InCommand.bCommandSuccessful = PlasticSourceControlUtils::RunUpdateStatus(Files, Operation->ShouldUpdateHistory(), InCommand.Concurrency, InCommand.ErrorMessages, States, InCommand.ChangesetNumber, InCommand.BranchName);
+		InCommand.bCommandSuccessful = PlasticSourceControlUtils::RunUpdateStatus(Files, Operation->ShouldUpdateHistory(), InCommand.ErrorMessages, States, InCommand.ChangesetNumber, InCommand.BranchName);
 		// Remove all "is not in a workspace" error and convert the result to "success" if there are no other errors
 		PlasticSourceControlUtils::RemoveRedundantErrors(InCommand, TEXT("is not in a workspace."));
 		if (!InCommand.bCommandSuccessful)
 		{
+			TArray<FString> ErrorMessages;
+			TArray<FString> Parameters;
+			FString BranchName, RepositoryName, ServerUrl;
+			if (PlasticSourceControlUtils::GetWorkspaceInfo(BranchName, RepositoryName, ServerUrl, ErrorMessages))
+			{
+				Parameters.Add(FString::Printf(TEXT("--server=%s"), *ServerUrl));
+			}
 			UE_LOG(LogSourceControl, Error, TEXT("FPlasticUpdateStatusWorker(ErrorMessages.Num()=%d) => checkconnection"), InCommand.ErrorMessages.Num());
 			// In case of error, execute a 'checkconnection' command to check the connectivity of the server.
-			InCommand.bConnectionDropped = !PlasticSourceControlUtils::RunCommand(TEXT("checkconnection"), TArray<FString>(), TArray<FString>(), InCommand.Concurrency, InCommand.InfoMessages, InCommand.ErrorMessages);
+			InCommand.bConnectionDropped = !PlasticSourceControlUtils::RunCommand(TEXT("checkconnection"), Parameters, TArray<FString>(), InCommand.InfoMessages, InCommand.ErrorMessages);
 			return false;
 		}
 
@@ -765,7 +861,8 @@ bool FPlasticUpdateStatusWorker::Execute(FPlasticSourceControlCommand& InCommand
 			{
 				if (State.IsConflicted())
 				{
-					// In case of a merge conflict, we need to put the tip of the "remote branch" on top of the history
+					// In case of a merge conflict, the Editor expects the tip of the "source (remote)" branch to be at the top of the history (index 0)
+					// as a way to represent the "merge in progress" in a 1D graph of the current branch "target (local)"
 					UE_LOG(LogSourceControl, Log, TEXT("%s: PendingMergeSourceChangeset %d"), *State.LocalFilename, State.PendingMergeSourceChangeset);
 					for (int32 IdxRevision = 0; IdxRevision < State.History.Num(); IdxRevision++)
 					{
@@ -802,7 +899,7 @@ bool FPlasticUpdateStatusWorker::Execute(FPlasticSourceControlCommand& InCommand
 	{
 		TArray<FString> ProjectDirs;
 		ProjectDirs.Add(FPaths::ConvertRelativePathToFull(FPaths::ProjectContentDir()));
-		InCommand.bCommandSuccessful = PlasticSourceControlUtils::RunUpdateStatus(ProjectDirs, Operation->ShouldUpdateHistory(), InCommand.Concurrency, InCommand.ErrorMessages, States, InCommand.ChangesetNumber, InCommand.BranchName);
+		InCommand.bCommandSuccessful = PlasticSourceControlUtils::RunUpdateStatus(ProjectDirs, Operation->ShouldUpdateHistory(), InCommand.ErrorMessages, States, InCommand.ChangesetNumber, InCommand.BranchName);
 	}
 	else
 	{
@@ -822,6 +919,23 @@ bool FPlasticUpdateStatusWorker::UpdateStates()
 {
 	TRACE_CPUPROFILER_EVENT_SCOPE(FPlasticUpdateStatusWorker::UpdateStates);
 
+	// Update affected changelists if any (in case of a file reverted outside of the Unreal Editor)
+	for (const FPlasticSourceControlState& NewState : States)
+	{
+		if (!NewState.IsModified())
+		{
+			TSharedRef<FPlasticSourceControlState, ESPMode::ThreadSafe> State = GetProvider().GetStateInternal(NewState.GetFilename());
+			if (State->Changelist.IsInitialized())
+			{
+				// 1- Remove these files from their previous changelist
+				TSharedRef<FPlasticSourceControlChangelistState, ESPMode::ThreadSafe> PreviousChangelist = GetProvider().GetStateInternal(State->Changelist);
+				PreviousChangelist->Files.Remove(State);
+				// 2- And reset the reference to their previous changelist
+				State->Changelist.Reset();
+			}
+		}
+	}
+
 	return PlasticSourceControlUtils::UpdateCachedStates(MoveTemp(States));
 }
 
@@ -842,7 +956,7 @@ bool FPlasticCopyWorker::Execute(FPlasticSourceControlCommand& InCommand)
 		const FString& Origin = InCommand.Files[0];
 		const FString Destination = FPaths::ConvertRelativePathToFull(Operation->GetDestination());
 
-		// Detects if the operation is a Move/Rename "Branch" that we need to track (or just a Duplicate/Copy, already Added to source control)
+		// Branch: The new file is branched from the original file (vs Add: The new file has no relation to the original file)
 		const bool bIsMoveOperation = (Operation->CopyMethod == FCopy::ECopyMethod::Branch);
 		if (bIsMoveOperation)
 		{
@@ -852,7 +966,7 @@ bool FPlasticCopyWorker::Execute(FPlasticSourceControlCommand& InCommand)
 			{
 				TArray<FString> DestinationFiles;
 				DestinationFiles.Add(Destination);
-				InCommand.bCommandSuccessful = PlasticSourceControlUtils::RunCommand(TEXT("undochange"), TArray<FString>(), DestinationFiles, InCommand.Concurrency, InCommand.InfoMessages, InCommand.ErrorMessages);
+				InCommand.bCommandSuccessful = PlasticSourceControlUtils::RunCommand(TEXT("undochange"), TArray<FString>(), DestinationFiles, InCommand.InfoMessages, InCommand.ErrorMessages);
 			}
 			// - execute a 'cm move --nomoveondisk' command to the destination to tell cm what happened
 			if (InCommand.bCommandSuccessful)
@@ -865,11 +979,11 @@ bool FPlasticCopyWorker::Execute(FPlasticSourceControlCommand& InCommand)
 				// Detect special case for a partial checkout (CS:-1 in Gluon mode)!
 				if (-1 != InCommand.ChangesetNumber)
 				{
-					InCommand.bCommandSuccessful = PlasticSourceControlUtils::RunCommand(TEXT("move"), Parameters, Files, InCommand.Concurrency, InCommand.InfoMessages, InCommand.ErrorMessages);
+					InCommand.bCommandSuccessful = PlasticSourceControlUtils::RunCommand(TEXT("move"), Parameters, Files, InCommand.InfoMessages, InCommand.ErrorMessages);
 				}
 				else
 				{
-					InCommand.bCommandSuccessful = PlasticSourceControlUtils::RunCommand(TEXT("partial move"), Parameters, Files, InCommand.Concurrency, InCommand.InfoMessages, InCommand.ErrorMessages);
+					InCommand.bCommandSuccessful = PlasticSourceControlUtils::RunCommand(TEXT("partial move"), Parameters, Files, InCommand.InfoMessages, InCommand.ErrorMessages);
 				}
 			}
 			// - add the redirector file (if it exists) to source control
@@ -880,11 +994,11 @@ bool FPlasticCopyWorker::Execute(FPlasticSourceControlCommand& InCommand)
 				// Detect special case for a partial checkout (CS:-1 in Gluon mode)!
 				if (-1 != InCommand.ChangesetNumber)
 				{
-					InCommand.bCommandSuccessful = PlasticSourceControlUtils::RunCommand(TEXT("add"), TArray<FString>(), Files, InCommand.Concurrency, InCommand.InfoMessages, InCommand.ErrorMessages);
+					InCommand.bCommandSuccessful = PlasticSourceControlUtils::RunCommand(TEXT("add"), TArray<FString>(), Files, InCommand.InfoMessages, InCommand.ErrorMessages);
 				}
 				else
 				{
-					InCommand.bCommandSuccessful = PlasticSourceControlUtils::RunCommand(TEXT("partial add"), TArray<FString>(), Files, InCommand.Concurrency, InCommand.InfoMessages, InCommand.ErrorMessages);
+					InCommand.bCommandSuccessful = PlasticSourceControlUtils::RunCommand(TEXT("partial add"), TArray<FString>(), Files, InCommand.InfoMessages, InCommand.ErrorMessages);
 				}
 			}
 		}
@@ -898,7 +1012,7 @@ bool FPlasticCopyWorker::Execute(FPlasticSourceControlCommand& InCommand)
 		TArray<FString> BothFiles;
 		BothFiles.Add(Origin);
 		BothFiles.Add(Destination);
-		PlasticSourceControlUtils::RunUpdateStatus(BothFiles, false, InCommand.Concurrency, InCommand.ErrorMessages, States, InCommand.ChangesetNumber, InCommand.BranchName);
+		PlasticSourceControlUtils::RunUpdateStatus(BothFiles, false, InCommand.ErrorMessages, States, InCommand.ChangesetNumber, InCommand.BranchName);
 	}
 	else
 	{
@@ -924,50 +1038,19 @@ bool FPlasticSyncWorker::Execute(FPlasticSourceControlCommand& InCommand)
 {
 	TRACE_CPUPROFILER_EVENT_SCOPE(FPlasticSyncWorker::Execute);
 
-	check(InCommand.Operation->GetName() == GetName());
+	TArray<FString> UpdatedFiles;
+	InCommand.bCommandSuccessful = PlasticSourceControlUtils::RunUpdate(InCommand.Files, GetProvider().IsPartialWorkspace(), UpdatedFiles, InCommand.ErrorMessages);
 
-	// Update specified directory to the head of the repository
-	// Detect special case for a partial checkout (CS:-1 in Gluon mode)!
-	if (-1 != InCommand.ChangesetNumber)
+	// now update the status of the updated files
+	if (UpdatedFiles.Num())
 	{
-		TSharedRef<FSync, ESPMode::ThreadSafe> Operation = StaticCastSharedRef<FSync>(InCommand.Operation);
-
-		TArray<FString> Parameters;
-
-		if (Operation->IsForced())
-		{
-			Parameters.Add(TEXT("--forced"));
-		}
-
-		if (Operation->GetRevision().IsEmpty())
-		{
-			Parameters.Add(TEXT("--last"));
-			Parameters.Add(TEXT("--dontmerge"));
-		}
-		else
-		{
-			Parameters.Add(FString::Printf(TEXT("--changeset=%s"), *Operation->GetRevision()));
-		}
-
-		InCommand.bCommandSuccessful = PlasticSourceControlUtils::RunCommand(TEXT("update"), Parameters, InCommand.Files, InCommand.Concurrency, InCommand.InfoMessages, InCommand.ErrorMessages);
-	}
-	else
-	{
-		InCommand.bCommandSuccessful = PlasticSourceControlUtils::RunCommand(TEXT("partial update"), {}, InCommand.Files, InCommand.Concurrency, InCommand.InfoMessages, InCommand.ErrorMessages);
+		PlasticSourceControlUtils::RunUpdateStatus(UpdatedFiles, false, InCommand.ErrorMessages, States, InCommand.ChangesetNumber, InCommand.BranchName);
 	}
 
-	if (InCommand.bCommandSuccessful)
+	if ((InCommand.Operation->GetName() == FName("SyncAll")))
 	{
-		// now update the status of our files
-		// detect the special case of a Sync of the root folder:
-		if ((InCommand.Files.Num() == 1) && (InCommand.Files.Last() == InCommand.PathToWorkspaceRoot))
-		{
-			// only update the status of assets in the Content directory
-			TArray<FString> ContentDir;
-			ContentDir.Add(FPaths::ConvertRelativePathToFull(FPaths::ProjectContentDir()));
-			PlasticSourceControlUtils::RunUpdateStatus(ContentDir, false, InCommand.Concurrency, InCommand.ErrorMessages, States, InCommand.ChangesetNumber, InCommand.BranchName);
-		}
-		// else: optim, no need to update the status of our files since this is done immediately after by the Editor
+		TSharedRef<FPlasticSyncAll, ESPMode::ThreadSafe> Operation = StaticCastSharedRef<FPlasticSyncAll>(InCommand.Operation);
+		Operation->UpdatedFiles = MoveTemp(UpdatedFiles);
 	}
 
 	return InCommand.bCommandSuccessful;
@@ -1011,11 +1094,11 @@ bool FPlasticResolveWorker::Execute(FPlasticSourceControlCommand& InCommand)
 		UE_LOG(LogSourceControl, Log, TEXT("resolve %s"), *State->PendingMergeFilename);
 
 		// Mark the conflicted file as resolved
-		InCommand.bCommandSuccessful = PlasticSourceControlUtils::RunCommand(TEXT("merge"), Parameters, OneFile, InCommand.Concurrency, InCommand.InfoMessages, InCommand.ErrorMessages);
+		InCommand.bCommandSuccessful = PlasticSourceControlUtils::RunCommand(TEXT("merge"), Parameters, OneFile, InCommand.InfoMessages, InCommand.ErrorMessages);
 	}
 
 	// now update the status of our files
-	PlasticSourceControlUtils::RunUpdateStatus(InCommand.Files, false, InCommand.Concurrency, InCommand.ErrorMessages, States, InCommand.ChangesetNumber, InCommand.BranchName);
+	PlasticSourceControlUtils::RunUpdateStatus(InCommand.Files, false, InCommand.ErrorMessages, States, InCommand.ChangesetNumber, InCommand.BranchName);
 
 	return InCommand.bCommandSuccessful;
 }
@@ -1039,7 +1122,7 @@ bool FPlasticGetPendingChangelistsWorker::Execute(FPlasticSourceControlCommand&
 	check(InCommand.Operation->GetName() == GetName());
 	TSharedRef<FUpdatePendingChangelistsStatus, ESPMode::ThreadSafe> Operation = StaticCastSharedRef<FUpdatePendingChangelistsStatus>(InCommand.Operation);
 
-	InCommand.bCommandSuccessful = PlasticSourceControlUtils::RunGetChangelists(InCommand.Concurrency, OutChangelistsStates, OutCLFilesStates, InCommand.ErrorMessages);
+	InCommand.bCommandSuccessful = PlasticSourceControlUtils::RunGetChangelists(OutChangelistsStates, OutCLFilesStates, InCommand.ErrorMessages);
 	if (InCommand.bCommandSuccessful)
 	{
 		// Remove the changelist that were not requested by the user.
@@ -1055,6 +1138,8 @@ bool FPlasticGetPendingChangelistsWorker::Execute(FPlasticSourceControlCommand&
 						});
 				});
 		}
+
+		InCommand.bCommandSuccessful = PlasticSourceControlUtils::RunGetShelves(OutChangelistsStates, InCommand.ErrorMessages);
 	}
 
 	bCleanupCache = InCommand.bCommandSuccessful;
@@ -1088,11 +1173,6 @@ bool FPlasticGetPendingChangelistsWorker::UpdateStates()
 			for (const auto& FileState : OutCLFilesStates[StatusIndex])
 			{
 				TSharedRef<FPlasticSourceControlState, ESPMode::ThreadSafe> CachedFileState = GetProvider().GetStateInternal(FileState.LocalFilename);
-				// Don't override "fileinfo" information and the potential LockedByOther state
-				if (CachedFileState->WorkspaceState != EWorkspaceState::LockedByOther)
-				{
-					CachedFileState->WorkspaceState = FileState.WorkspaceState;
-				}
 				CachedFileState->Changelist = CLStatus.Changelist;
 				ChangelistState->Files.AddUnique(CachedFileState);
 			}
@@ -1143,19 +1223,19 @@ FPlasticSourceControlChangelist GenerateUniqueChangelistName(FPlasticSourceContr
 	return NewChangelist;
 }
 
-FPlasticSourceControlChangelist CreatePendingChangelist(FPlasticSourceControlProvider& PlasticSourceControlProvider, const FString& InDescription, EConcurrency::Type InConcurrency, TArray<FString>& InInfoMessages, TArray<FString>& InErrorMessages)
+FPlasticSourceControlChangelist CreatePendingChangelist(FPlasticSourceControlProvider& PlasticSourceControlProvider, const FString& InDescription, TArray<FString>& InInfoMessages, TArray<FString>& InErrorMessages)
 {
 	FPlasticSourceControlChangelist NewChangelist = GenerateUniqueChangelistName(PlasticSourceControlProvider);
 
 	bool bCommandSuccessful;
 	TArray<FString> Parameters;
-	if (PlasticSourceControlProvider.GetPlasticScmVersion() < s_NewChangelistFileArgsPlasticScmVersion)
+	if (PlasticSourceControlProvider.GetPlasticScmVersion() < PlasticSourceControlVersions::NewChangelistFileArgs)
 	{
 		Parameters.Add(TEXT("add"));
 		Parameters.Add(TEXT("\"") + NewChangelist.GetName() + TEXT("\""));
 		Parameters.Add(TEXT("\"") + InDescription + TEXT("\""));
 		Parameters.Add(TEXT("--persistent")); // Create a persistent changelist to stay close to Perforce behavior
-		bCommandSuccessful = PlasticSourceControlUtils::RunCommand(TEXT("changelist"), Parameters, TArray<FString>(), InConcurrency, InInfoMessages, InErrorMessages);
+		bCommandSuccessful = PlasticSourceControlUtils::RunCommand(TEXT("changelist"), Parameters, TArray<FString>(), InInfoMessages, InErrorMessages);
 	}
 	else
 	{
@@ -1166,7 +1246,7 @@ FPlasticSourceControlChangelist CreatePendingChangelist(FPlasticSourceControlPro
 		Parameters.Add(FString::Printf(TEXT("--descriptionfile=\"%s\""), *FPaths::ConvertRelativePathToFull(ChangelistDescriptionFile.GetFilename())));
 		Parameters.Add(TEXT("--persistent")); // Create a persistent changelist to stay close to Perforce behavior
 		UE_LOG(LogSourceControl, Verbose, TEXT("CreatePendingChangelist(%s):\n\"%s\""), *NewChangelist.GetName(), *InDescription);
-		bCommandSuccessful = PlasticSourceControlUtils::RunCommand(TEXT("changelist"), Parameters, TArray<FString>(), InConcurrency, InInfoMessages, InErrorMessages);
+		bCommandSuccessful = PlasticSourceControlUtils::RunCommand(TEXT("changelist"), Parameters, TArray<FString>(), InInfoMessages, InErrorMessages);
 	}
 	if (!bCommandSuccessful)
 	{
@@ -1176,16 +1256,16 @@ FPlasticSourceControlChangelist CreatePendingChangelist(FPlasticSourceControlPro
 	return NewChangelist;
 }
 
-bool EditChangelistDescription(const FPlasticSourceControlProvider& PlasticSourceControlProvider, const FPlasticSourceControlChangelist& InChangelist, const FString& InDescription, EConcurrency::Type InConcurrency, TArray<FString>& InInfoMessages, TArray<FString>& InErrorMessages)
+bool EditChangelistDescription(const FPlasticSourceControlProvider& PlasticSourceControlProvider, const FPlasticSourceControlChangelist& InChangelist, const FString& InDescription, TArray<FString>& InInfoMessages, TArray<FString>& InErrorMessages)
 {
 	TArray<FString> Parameters;
 	Parameters.Add(TEXT("edit"));
-	if (PlasticSourceControlProvider.GetPlasticScmVersion() < s_NewChangelistFileArgsPlasticScmVersion)
+	if (PlasticSourceControlProvider.GetPlasticScmVersion() < PlasticSourceControlVersions::NewChangelistFileArgs)
 	{
 		Parameters.Add(TEXT("\"") + InChangelist.GetName() + TEXT("\""));
 		Parameters.Add(TEXT("description"));
 		Parameters.Add(TEXT("\"") + InDescription + TEXT("\""));
-		return PlasticSourceControlUtils::RunCommand(TEXT("changelist"), Parameters, TArray<FString>(), InConcurrency, InInfoMessages, InErrorMessages);
+		return PlasticSourceControlUtils::RunCommand(TEXT("changelist"), Parameters, TArray<FString>(), InInfoMessages, InErrorMessages);
 	}
 	else
 	{
@@ -1195,20 +1275,20 @@ bool EditChangelistDescription(const FPlasticSourceControlProvider& PlasticSourc
 		const FScopedTempFile ChangelistDescriptionFile(InDescription);
 		Parameters.Add(FString::Printf(TEXT("--descriptionfile=\"%s\""), *FPaths::ConvertRelativePathToFull(ChangelistDescriptionFile.GetFilename())));
 		UE_LOG(LogSourceControl, Verbose, TEXT("EditChangelistDescription(%s\n%s)"), *InChangelist.GetName(), *InDescription);
-		return PlasticSourceControlUtils::RunCommand(TEXT("changelist"), Parameters, TArray<FString>(), InConcurrency, InInfoMessages, InErrorMessages);
+		return PlasticSourceControlUtils::RunCommand(TEXT("changelist"), Parameters, TArray<FString>(), InInfoMessages, InErrorMessages);
 	}
 }
 
-bool MoveFilesToChangelist(const FPlasticSourceControlProvider& PlasticSourceControlProvider, const FPlasticSourceControlChangelist& InChangelist, const TArray<FString>& InFiles, const EConcurrency::Type InConcurrency, TArray<FString>& OutResults, TArray<FString>& OutErrorMessages)
+bool MoveFilesToChangelist(const FPlasticSourceControlProvider& PlasticSourceControlProvider, const FPlasticSourceControlChangelist& InChangelist, const TArray<FString>& InFiles, TArray<FString>& OutResults, TArray<FString>& OutErrorMessages)
 {
 	if (InFiles.Num() > 0)
 	{
 		TArray<FString> Parameters;
-		if (PlasticSourceControlProvider.GetPlasticScmVersion() < s_NewChangelistFileArgsPlasticScmVersion)
+		if (PlasticSourceControlProvider.GetPlasticScmVersion() < PlasticSourceControlVersions::NewChangelistFileArgs)
 		{
 			Parameters.Add(TEXT("\"") + InChangelist.GetName() + TEXT("\""));
 			Parameters.Add(TEXT("add"));
-			return PlasticSourceControlUtils::RunCommand(TEXT("changelist"), Parameters, InFiles, InConcurrency, OutResults, OutErrorMessages);
+			return PlasticSourceControlUtils::RunCommand(TEXT("changelist"), Parameters, InFiles, OutResults, OutErrorMessages);
 		}
 		else
 		{
@@ -1216,7 +1296,7 @@ bool MoveFilesToChangelist(const FPlasticSourceControlProvider& PlasticSourceCon
 			Parameters.Add(FString::Printf(TEXT("--namefile=\"%s\""), *FPaths::ConvertRelativePathToFull(ChangelistNameFile.GetFilename())));
 			Parameters.Add(TEXT("add"));
 			UE_LOG(LogSourceControl, Verbose, TEXT("MoveFilesToChangelist(%s)"), *InChangelist.GetName());
-			return PlasticSourceControlUtils::RunCommand(TEXT("changelist"), Parameters, InFiles, InConcurrency, OutResults, OutErrorMessages);
+			return PlasticSourceControlUtils::RunCommand(TEXT("changelist"), Parameters, InFiles, OutResults, OutErrorMessages);
 		}
 	}
 	return true;
@@ -1242,7 +1322,7 @@ bool FPlasticNewChangelistWorker::Execute(class FPlasticSourceControlCommand& In
 
 	FString Description = Operation->GetDescription().ToString();
 	// Note: old "cm" doesn't support newlines, quotes, and question marks on changelist's name or description
-	if (GetProvider().GetPlasticScmVersion() < s_NewChangelistFileArgsPlasticScmVersion)
+	if (GetProvider().GetPlasticScmVersion() < PlasticSourceControlVersions::NewChangelistFileArgs)
 	{
 		Description.ReplaceInline(TEXT("\r\n"), TEXT(" "), ESearchCase::CaseSensitive);
 		Description.ReplaceCharInline(TEXT('\n'), TEXT(' '), ESearchCase::CaseSensitive);
@@ -1252,7 +1332,7 @@ bool FPlasticNewChangelistWorker::Execute(class FPlasticSourceControlCommand& In
 	}
 
 	// Create a new numbered persistent changelist ala Perforce
-	NewChangelist = CreatePendingChangelist(GetProvider(), Description, InCommand.Concurrency, InCommand.InfoMessages, InCommand.ErrorMessages);
+	NewChangelist = CreatePendingChangelist(GetProvider(), Description, InCommand.InfoMessages, InCommand.ErrorMessages);
 
 	// Successfully created new changelist
 	if (NewChangelist.IsInitialized())
@@ -1266,7 +1346,7 @@ bool FPlasticNewChangelistWorker::Execute(class FPlasticSourceControlCommand& In
 
 		if (InCommand.Files.Num() > 0)
 		{
-			InCommand.bCommandSuccessful = MoveFilesToChangelist(GetProvider(), NewChangelist, InCommand.Files, InCommand.Concurrency, InCommand.InfoMessages, InCommand.ErrorMessages);
+			InCommand.bCommandSuccessful = MoveFilesToChangelist(GetProvider(), NewChangelist, InCommand.Files, InCommand.InfoMessages, InCommand.ErrorMessages);
 			if (InCommand.bCommandSuccessful)
 			{
 				MovedFiles = InCommand.Files;
@@ -1332,14 +1412,14 @@ bool FPlasticDeleteChangelistWorker::Execute(class FPlasticSourceControlCommand&
 		check(InCommand.Operation->GetName() == GetName());
 		TSharedRef<FDeleteChangelist, ESPMode::ThreadSafe> Operation = StaticCastSharedRef<FDeleteChangelist>(InCommand.Operation);
 
-		InCommand.bCommandSuccessful = DeleteChangelist(GetProvider(), InCommand.Changelist, InCommand.Concurrency, InCommand.InfoMessages, InCommand.ErrorMessages);
+		InCommand.bCommandSuccessful = DeleteChangelist(GetProvider(), InCommand.Changelist, InCommand.InfoMessages, InCommand.ErrorMessages);
 
 		// NOTE: for now it is not possible to delete a changelist with files through the Editor
 		if (InCommand.Files.Num() > 0 && InCommand.bCommandSuccessful)
 		{
 			TSharedRef<FPlasticSourceControlChangelistState, ESPMode::ThreadSafe> ChangelistState = GetProvider().GetStateInternal(InCommand.Changelist);
 			const TArray<FString> Files = FileNamesFromFileStates(ChangelistState->Files);
-			InCommand.bCommandSuccessful = MoveFilesToChangelist(GetProvider(), FPlasticSourceControlChangelist::DefaultChangelist, Files, InCommand.Concurrency, InCommand.InfoMessages, InCommand.ErrorMessages);
+			InCommand.bCommandSuccessful = MoveFilesToChangelist(GetProvider(), FPlasticSourceControlChangelist::DefaultChangelist, Files, InCommand.InfoMessages, InCommand.ErrorMessages);
 		}
 
 		// Keep track of changelist to update the cache
@@ -1379,7 +1459,7 @@ bool FPlasticEditChangelistWorker::Execute(class FPlasticSourceControlCommand& I
 
 	EditedDescription = Operation->GetDescription().ToString();
 	// Note: old "cm" doesn't support newlines, quotes, and question marks on changelist's name or description
-	if (GetProvider().GetPlasticScmVersion() < s_NewChangelistFileArgsPlasticScmVersion)
+	if (GetProvider().GetPlasticScmVersion() < PlasticSourceControlVersions::NewChangelistFileArgs)
 	{
 		EditedDescription.ReplaceInline(TEXT("\r\n"), TEXT(" "), ESearchCase::CaseSensitive);
 		EditedDescription.ReplaceCharInline(TEXT('\n'), TEXT(' '), ESearchCase::CaseSensitive);
@@ -1391,18 +1471,18 @@ bool FPlasticEditChangelistWorker::Execute(class FPlasticSourceControlCommand& I
 	if (InCommand.Changelist.IsDefault())
 	{
 		// Create a new numbered persistent changelist since we cannot edit the default changelist
-		EditedChangelist = CreatePendingChangelist(GetProvider(), EditedDescription, InCommand.Concurrency, InCommand.InfoMessages, InCommand.ErrorMessages);
+		EditedChangelist = CreatePendingChangelist(GetProvider(), EditedDescription, InCommand.InfoMessages, InCommand.ErrorMessages);
 		if (EditedChangelist.IsInitialized())
 		{
 			// And then move all its files to the new changelist
 			TSharedRef<FPlasticSourceControlChangelistState, ESPMode::ThreadSafe> ChangelistState = GetProvider().GetStateInternal(InCommand.Changelist);
 			ReopenedFiles = FileNamesFromFileStates(ChangelistState->Files);
-			InCommand.bCommandSuccessful = MoveFilesToChangelist(GetProvider(), EditedChangelist, ReopenedFiles, InCommand.Concurrency, InCommand.InfoMessages, InCommand.ErrorMessages);
+			InCommand.bCommandSuccessful = MoveFilesToChangelist(GetProvider(), EditedChangelist, ReopenedFiles, InCommand.InfoMessages, InCommand.ErrorMessages);
 		}
 	}
 	else
 	{
-		InCommand.bCommandSuccessful = EditChangelistDescription(GetProvider(), InCommand.Changelist, EditedDescription, InCommand.Concurrency, InCommand.InfoMessages, InCommand.ErrorMessages);
+		InCommand.bCommandSuccessful = EditChangelistDescription(GetProvider(), InCommand.Changelist, EditedDescription, InCommand.InfoMessages, InCommand.ErrorMessages);
 		if (InCommand.bCommandSuccessful)
 		{
 			EditedChangelist = InCommand.Changelist;
@@ -1461,7 +1541,7 @@ bool FPlasticReopenWorker::Execute(FPlasticSourceControlCommand& InCommand)
 
 	check(InCommand.Operation->GetName() == GetName());
 
-	InCommand.bCommandSuccessful = MoveFilesToChangelist(GetProvider(), InCommand.Changelist, InCommand.Files, InCommand.Concurrency, InCommand.InfoMessages, InCommand.ErrorMessages);
+	InCommand.bCommandSuccessful = MoveFilesToChangelist(GetProvider(), InCommand.Changelist, InCommand.Files, InCommand.InfoMessages, InCommand.ErrorMessages);
 	if (InCommand.bCommandSuccessful)
 	{
 		ReopenedFiles = InCommand.Files;
@@ -1504,4 +1584,391 @@ bool FPlasticReopenWorker::UpdateStates()
 		return false;
 	}
 }
+
+bool CreateShelve(const FString& InChangelistName, const FString& InChangelistDescription, const TArray<FString>& InFilesToShelve, int32& OutShelveId, TArray<FString>& OutErrorMessages)
+{
+	TArray<FString> Results;
+	TArray<FString> Parameters;
+	const FString ShelveDescription = FString::Printf(TEXT("Changelist%s: %s"), *InChangelistName, *InChangelistDescription);
+	const FScopedTempFile CommentsFile(ShelveDescription);
+	Parameters.Add(TEXT("create"));
+	Parameters.Add(FString::Printf(TEXT("-commentsfile=\"%s\""), *FPaths::ConvertRelativePathToFull(CommentsFile.GetFilename())));
+	const bool bCommandSuccessful = PlasticSourceControlUtils::RunCommand(TEXT("shelveset"), Parameters, InFilesToShelve, Results, OutErrorMessages);
+	if (bCommandSuccessful && Results.Num() > 0)
+	{
+		// Parse the result to update the changelist state with the id of the shelve
+		// "Created shelve sh:12@UE5PlasticPluginDev@test@cloud (mount:'/')"
+		FString CreatedShelveStatus = Results[Results.Num() - 1];
+		if (CreatedShelveStatus.StartsWith(TEXT("Created shelve sh:")))
+		{
+			CreatedShelveStatus.RightChopInline(18);
+			int32 SeparatorIndex;
+			if (CreatedShelveStatus.FindChar(TEXT('@'), SeparatorIndex))
+			{
+				CreatedShelveStatus.LeftInline(SeparatorIndex);
+				OutShelveId = FCString::Atoi(*CreatedShelveStatus);
+			}
+		}
+	}
+
+	return OutShelveId != ISourceControlState::INVALID_REVISION;
+}
+
+bool DeleteShelve(const int32 InShelveId, TArray<FString>& OutErrorMessages)
+{
+	TArray<FString> Results;
+	TArray<FString> Parameters;
+	Parameters.Add(TEXT("delete"));
+	Parameters.Add(FString::Printf(TEXT("sh:%d"), InShelveId));
+	return PlasticSourceControlUtils::RunCommand(TEXT("shelveset"), Parameters, TArray<FString>(), Results, OutErrorMessages);
+}
+
+FName FPlasticShelveWorker::GetName() const
+{
+	return "Shelve";
+}
+
+bool FPlasticShelveWorker::Execute(FPlasticSourceControlCommand& InCommand)
+{
+	TRACE_CPUPROFILER_EVENT_SCOPE(FPlasticShelveWorker::Execute);
+
+	check(InCommand.Operation->GetName() == GetName());
+	TSharedRef<FShelve, ESPMode::ThreadSafe> Operation = StaticCastSharedRef<FShelve>(InCommand.Operation);
+
+	FPlasticSourceControlChangelist Changelist(InCommand.Changelist);
+
+	TArray<FString> FilesToShelve = InCommand.Files;
+
+	int32 PreviousShelveId = ISourceControlState::INVALID_REVISION;
+
+	InCommand.bCommandSuccessful = true;
+
+	if (InCommand.Changelist.IsInitialized())
+	{
+		TSharedRef<FPlasticSourceControlChangelistState, ESPMode::ThreadSafe> ChangelistState = GetProvider().GetStateInternal(InCommand.Changelist);
+
+		PreviousShelveId = ChangelistState->ShelveId;
+
+		// If the command has specified a changelist but no files, then get all files from it
+		if (FilesToShelve.Num() == 0)
+		{
+			FilesToShelve = FileNamesFromFileStates(ChangelistState->Files);
+		}
+		// If the command has specified some files, and if there was already a shelve,
+		// ensure that the previous shelved files are also put back into the new one.
+		// This is to mimic Perforce behavior, where we can add files to a shelve after we made more changes.
+		// NOTE: this is a workaround for the fact that we cannot edit an existing shelve.
+		else if (ChangelistState->ShelvedFiles.Num() > 0)
+		{
+			for (const FSourceControlStateRef& ShelveFile : ChangelistState->ShelvedFiles)
+			{
+				FilesToShelve.AddUnique(ShelveFile->GetFilename());
+			}
+		}
+
+		// Ensure that all the files to shelve are indeed in the changelist, so that we can actually shelve them
+		// NOTE: this is because of the workaround that requires to create a new shelve and delete the old one
+		for (FString& FileToShelve : FilesToShelve)
+		{
+			if (!ChangelistState->Files.ContainsByPredicate([&FileToShelve](const FSourceControlStateRef& FileState) { return FileToShelve == FileState->GetFilename(); }))
+			{
+				FPaths::MakePathRelativeTo(FileToShelve, *FPaths::ProjectDir());
+				UE_LOG(LogSourceControl, Error, TEXT("The file /%s is not in the changelist anymore, so the shelve cannot be updated. Unshelve the corresponding change and retry."), *FileToShelve);
+				InCommand.bCommandSuccessful = false;
+			}
+		}
+	}
+
+	if (InCommand.bCommandSuccessful)
+	{
+		// If the command is issued on the default changelist, then we need to create a new changelist,
+		// move the files to the new changelist, then shelve the files
+		if (InCommand.Changelist.IsDefault())
+		{
+			// Create a new numbered persistent changelist ala Perforce
+			Changelist = CreatePendingChangelist(GetProvider(), Operation->GetDescription().ToString(), InCommand.InfoMessages, InCommand.ErrorMessages);
+			if (Changelist.IsInitialized())
+			{
+				InCommand.bCommandSuccessful = MoveFilesToChangelist(GetProvider(), Changelist, FilesToShelve, InCommand.InfoMessages, InCommand.ErrorMessages);
+				if (InCommand.bCommandSuccessful)
+				{
+					MovedFiles = FilesToShelve;
+				}
+			}
+		}
+	}
+
+	if (InCommand.bCommandSuccessful)
+	{
+		ChangelistDescription = *Operation->GetDescription().ToString();
+
+		InCommand.bCommandSuccessful = CreateShelve(Changelist.GetName(), ChangelistDescription, FilesToShelve, ShelveId, InCommand.ErrorMessages);
+		if (InCommand.bCommandSuccessful)
+		{
+			InChangelistToUpdate = InCommand.Changelist;
+			OutChangelistToUpdate = Changelist;
+			ShelvedFiles = FilesToShelve;
+
+			// If there was already a shelve, we have now created a new one with updated files, so we must delete the old one
+			if (PreviousShelveId != ISourceControlState::INVALID_REVISION)
+			{
+				DeleteShelve(PreviousShelveId, InCommand.ErrorMessages);
+			}
+		}
+		else
+		{
+			// In case of failure to shelve, if we had to create a new changelist, move the files back to the default changelist and delete the changelist
+			if (Changelist != InCommand.Changelist)
+			{
+				if (MovedFiles.Num() > 0)
+				{
+					MoveFilesToChangelist(GetProvider(), InCommand.Changelist, MovedFiles, InCommand.InfoMessages, InCommand.ErrorMessages);
+				}
+
+				DeleteChangelist(GetProvider(), Changelist, InCommand.InfoMessages, InCommand.ErrorMessages);
+			}
+		}
+	}
+
+	return InCommand.bCommandSuccessful;
+}
+
+bool FPlasticShelveWorker::UpdateStates()
+{
+	TRACE_CPUPROFILER_EVENT_SCOPE(FPlasticShelveWorker::UpdateStates);
+
+	if (OutChangelistToUpdate.IsInitialized())
+	{
+		TSharedRef<FPlasticSourceControlChangelistState, ESPMode::ThreadSafe> DestinationChangelistState = GetProvider().GetStateInternal(OutChangelistToUpdate);
+
+		bool bMovedFiles = false;
+
+		// If we moved files to a new changelist, then we must make sure that the files are properly moved
+		if (InChangelistToUpdate != OutChangelistToUpdate && MovedFiles.Num() > 0)
+		{
+			const FDateTime Now = FDateTime::Now();
+			TSharedRef<FPlasticSourceControlChangelistState, ESPMode::ThreadSafe> SourceChangelistState = GetProvider().GetStateInternal(InChangelistToUpdate);
+
+			DestinationChangelistState->Changelist = OutChangelistToUpdate;
+			DestinationChangelistState->Description = ChangelistDescription;
+
+			for (const FString& MovedFile : MovedFiles)
+			{
+				TSharedRef<FPlasticSourceControlState, ESPMode::ThreadSafe> FileState = GetProvider().GetStateInternal(MovedFile);
+
+				SourceChangelistState->Files.Remove(FileState);
+				DestinationChangelistState->Files.Add(FileState);
+				FileState->Changelist = OutChangelistToUpdate;
+				FileState->TimeStamp = Now;
+			}
+
+			bMovedFiles = true;
+		}
+
+		DestinationChangelistState->ShelveId = ShelveId;
+
+		// And finally, add the shelved files to the changelist state
+		DestinationChangelistState->ShelvedFiles.Reset(ShelvedFiles.Num());
+		for (FString ShelvedFile : ShelvedFiles)
+		{
+			TSharedRef<FPlasticSourceControlState, ESPMode::ThreadSafe> FileState = GetProvider().GetStateInternal(ShelvedFile);
+			PlasticSourceControlUtils::AddShelvedFileToChangelist(DestinationChangelistState.Get(), MoveTemp(ShelvedFile), FileState->WorkspaceState);
+		}
+
+		return bMovedFiles || ShelvedFiles.Num() > 0;
+	}
+	else
+	{
+		return false;
+	}
+}
+
+
+
+FName FPlasticUnshelveWorker::GetName() const
+{
+	return "Unshelve";
+}
+
+bool FPlasticUnshelveWorker::Execute(FPlasticSourceControlCommand& InCommand)
+{
+	TRACE_CPUPROFILER_EVENT_SCOPE(FPlasticUnshelveWorker::Execute);
+
+	// Get the state of the changelist to operate on
+	TSharedRef<FPlasticSourceControlChangelistState, ESPMode::ThreadSafe> ChangelistState = GetProvider().GetStateInternal(InCommand.Changelist);
+
+	InCommand.bCommandSuccessful = (ChangelistState->ShelveId != ISourceControlState::INVALID_REVISION);
+
+	// Detect if any file to unshelve has some local modification, which would fail the "unshelve" operation with a merge conflict
+	// NOTE: we could decide to automatically undo the local changes in order for this process to be automatic, like with Perforce
+	for (FString& File : InCommand.Files)
+	{
+		if (ChangelistState->Files.FindByPredicate([&File](const FSourceControlStateRef& FileState) { return FileState->GetFilename().Equals(File, ESearchCase::IgnoreCase); }))
+		{
+			FPaths::MakePathRelativeTo(File, *FPaths::ProjectDir());
+			UE_LOG(LogSourceControl, Error, TEXT("Revert /%s before unshelving the corresponding change from the shelve."), *File);
+			InCommand.bCommandSuccessful = false;
+		}
+	}
+
+	if (InCommand.bCommandSuccessful)
+	{
+		// Get the list of files to unshelve if not all of them are selected
+		TArray<FString> FilesToUnshelve;
+		if (InCommand.Files.Num() < ChangelistState->ShelvedFiles.Num())
+		{
+			if (GetProvider().GetPlasticScmVersion() < PlasticSourceControlVersions::ShelvesetApplySelection)
+			{
+				// On old version, don't unshelve the files if they are not all selected (since we couldn't apply only a selection of files from a shelve)
+				UE_LOG(LogSourceControl, Error,
+					TEXT("Plastic SCM %s cannot unshelve a selection of files from a shelve. Unshelve them all at once or update to %s or above."),
+					*GetProvider().GetPlasticScmVersion().String,
+					*PlasticSourceControlVersions::ShelvesetApplySelection.String
+				);
+				return false;
+			}
+			const FString PathToWorkspaceRoot = GetProvider().GetPathToWorkspaceRoot();
+			FilesToUnshelve.Reset(InCommand.Files.Num());
+			for (FString File : InCommand.Files)
+			{
+				// Make path relative the workspace root, since the shelveset apply operation require server paths
+				FPaths::MakePathRelativeTo(File, *PathToWorkspaceRoot);
+				FilesToUnshelve.Add(TEXT("/") + File);
+			}
+		}
+
+		// Make the Editor Unlink the assets if they are loaded in memory so that source control can override the corresponding files
+		PackageUtils::UnlinkPackagesInMainThread(InCommand.Files);
+
+		{
+			// 'cm shelveset apply sh:88 "/Content/Blueprints/BP_CheckedOut.uasset"'
+			TArray<FString> Parameters;
+			Parameters.Add(TEXT("apply"));
+			Parameters.Add(FString::Printf(TEXT("sh:%d"), ChangelistState->ShelveId));
+			InCommand.bCommandSuccessful = PlasticSourceControlUtils::RunCommand(TEXT("shelveset"), Parameters, FilesToUnshelve, InCommand.InfoMessages, InCommand.ErrorMessages);
+		}
+
+		// Reload packages that where updated by the Unshelve operation (and the current map if needed)
+		PackageUtils::ReloadPackagesInMainThread(InCommand.Files);
+	}
+
+	if (InCommand.bCommandSuccessful)
+	{
+		// move all the unshelved files back to the changelist
+		InCommand.bCommandSuccessful = MoveFilesToChangelist(GetProvider(), InCommand.Changelist, InCommand.Files, InCommand.InfoMessages, InCommand.ErrorMessages);
+	}
+
+	if (InCommand.bCommandSuccessful)
+	{
+		// now update the status of our files
+		PlasticSourceControlUtils::RunUpdateStatus(InCommand.Files, false, InCommand.ErrorMessages, States, InCommand.ChangesetNumber, InCommand.BranchName);
+
+		ChangelistToUpdate = InCommand.Changelist;
+	}
+
+	return InCommand.bCommandSuccessful;
+}
+
+bool FPlasticUnshelveWorker::UpdateStates()
+{
+	TRACE_CPUPROFILER_EVENT_SCOPE(FPlasticUnshelveWorker::UpdateStates);
+
+	UpdateChangelistState(GetProvider(), ChangelistToUpdate, States);
+
+	return PlasticSourceControlUtils::UpdateCachedStates(MoveTemp(States));
+}
+
+FName FPlasticDeleteShelveWorker::GetName() const
+{
+	return "DeleteShelved";
+}
+
+bool FPlasticDeleteShelveWorker::Execute(FPlasticSourceControlCommand& InCommand)
+{
+	TRACE_CPUPROFILER_EVENT_SCOPE(FPlasticDeleteShelveWorker::Execute);
+
+	TSharedRef<FPlasticSourceControlChangelistState, ESPMode::ThreadSafe> ChangelistState = GetProvider().GetStateInternal(InCommand.Changelist);
+
+	InCommand.bCommandSuccessful = (ChangelistState->ShelveId != ISourceControlState::INVALID_REVISION);
+
+	// Get the list of files to keep in the shelve if not all of them are selected (to create a new shelve with them)
+	TArray<FString> FilesToShelve;
+	if (InCommand.Files.Num() < ChangelistState->ShelvedFiles.Num())
+	{
+		for (const auto& ShelveState : ChangelistState->ShelvedFiles)
+		{
+			if (!InCommand.Files.Contains(ShelveState->GetFilename()))
+			{
+				FString File = ShelveState->GetFilename();
+				
+				// Check that all this files are is still in the corresponding changelist, else we won't be able to create the new shelve!
+				if (ChangelistState->Files.ContainsByPredicate([&File](FSourceControlStateRef& State) { return File == State->GetFilename(); }))
+				{
+					FilesToShelve.Add(MoveTemp(File));
+				}
+				else
+				{
+
+					FPaths::MakePathRelativeTo(File, *FPaths::ProjectDir());
+					UE_LOG(LogSourceControl, Error, TEXT("The file /%s is not in the changelist anymore, so the shelve cannot be updated. Unshelve the corresponding change and retry."), *File);
+					InCommand.bCommandSuccessful = false;
+				}
+			}
+		}
+	}
+
+	if (InCommand.bCommandSuccessful)
+	{
+		ChangelistToUpdate = InCommand.Changelist;
+		FilesToRemove = InCommand.Files;
+	}
+	
+	if (InCommand.bCommandSuccessful && FilesToShelve.Num() > 0)
+	{
+		// Create a new shelve with the other files
+		InCommand.bCommandSuccessful = CreateShelve(InCommand.Changelist.GetName(), ChangelistState->GetDescriptionText().ToString(), FilesToShelve, ShelveId, InCommand.ErrorMessages);
+	}
+
+	if (InCommand.bCommandSuccessful)
+	{
+		// Delete the old shelve
+		InCommand.bCommandSuccessful = DeleteShelve(ChangelistState->ShelveId, InCommand.ErrorMessages);
+	}
+
+	return InCommand.bCommandSuccessful;
+}
+
+bool FPlasticDeleteShelveWorker::UpdateStates()
+{
+	TRACE_CPUPROFILER_EVENT_SCOPE(FPlasticDeleteShelveWorker::UpdateStates);
+
+	if (ChangelistToUpdate.IsInitialized())
+	{
+		TSharedRef<FPlasticSourceControlChangelistState, ESPMode::ThreadSafe> ChangelistState = GetProvider().GetStateInternal(ChangelistToUpdate);
+
+		ChangelistState->ShelveId = ShelveId;
+
+		if (FilesToRemove.Num() > 0)
+		{
+			// NOTE: for now, Plastic SCM cannot delete a selection of files from a shelve, so FilesToRemove and this specific case aren't really needed (yet)
+			return ChangelistState->ShelvedFiles.RemoveAll([this](FSourceControlStateRef& State) -> bool
+				{
+					return Algo::AnyOf(FilesToRemove, [&State](const FString& File) {
+						return State->GetFilename() == File;
+					});
+				}) > 0;
+		}
+		else
+		{
+			bool bHadShelvedFiles = (ChangelistState->ShelvedFiles.Num() > 0);
+			ChangelistState->ShelvedFiles.Reset();
+			return bHadShelvedFiles;
+		}
+	}
+	else
+	{
+		return false;
+	}
+}
+
 #undef LOCTEXT_NAMESPACE
diff --git a/Engine/Plugins/Developer/PlasticSourceControl/Source/PlasticSourceControl/Private/PlasticSourceControlOperations.h b/Engine/Plugins/Developer/PlasticSourceControl/Source/PlasticSourceControl/Private/PlasticSourceControlOperations.h
--- a/Engine/Plugins/Developer/PlasticSourceControl/Source/PlasticSourceControl/Private/PlasticSourceControlOperations.h
+++ b/Engine/Plugins/Developer/PlasticSourceControl/Source/PlasticSourceControl/Private/PlasticSourceControlOperations.h
@@ -1,6 +1,6 @@
 // Copyright Epic Games, Inc. All Rights Reserved.
 
-// Specialization of classes defineds in Engine\Source\Developer\SourceControl\Public\SourceControlOperations.h
+// Specialization of classes defined in Engine\Source\Developer\SourceControl\Public\SourceControlOperations.h
 
 #pragma once
 
@@ -9,6 +9,7 @@
 #include "PlasticSourceControlRevision.h"
 #include "PlasticSourceControlState.h"
 #include "ISourceControlOperation.h"
+#include "SourceControlOperations.h"
 
 #include "PlasticSourceControlChangelist.h"
 #include "PlasticSourceControlChangelistState.h"
@@ -28,16 +29,36 @@ public:
 };
 
 
+/**
+ * Internal operation used to sync all files in the workspace
+*/
+class FPlasticSyncAll final : public FSync
+{
+public:
+	// ISourceControlOperation interface
+	virtual FName GetName() const override
+	{
+		return "SyncAll";
+	}
+
+	/** List of files updated by the operation */
+	TArray<FString> UpdatedFiles;
+};
+
+
 /**
  * Internal operation used to revert checked-out files
 */
-class FPlasticRevertAll final : public ISourceControlOperation
+class FPlasticRevertAll final : public FRevert
 {
 public:
 	// ISourceControlOperation interface
 	virtual FName GetName() const override;
 
 	virtual FText GetInProgressString() const override;
+
+	/** List of files updated by the operation */
+	TArray<FString> UpdatedFiles;
 };
 
 
@@ -399,3 +420,77 @@ protected:
 	FPlasticSourceControlChangelist DestinationChangelist;
 };
 
+class FPlasticShelveWorker final : public IPlasticSourceControlWorker
+{
+public:
+	explicit FPlasticShelveWorker(FPlasticSourceControlProvider& InSourceControlProvider)
+		: IPlasticSourceControlWorker(InSourceControlProvider)
+	{}
+	virtual ~FPlasticShelveWorker() = default;
+
+	// IPlasticSourceControlWorker interface
+	virtual FName GetName() const override;
+	virtual bool Execute(FPlasticSourceControlCommand& InCommand) override;
+	virtual bool UpdateStates() override;
+
+protected:
+	int32 ShelveId = ISourceControlState::INVALID_REVISION;
+
+	TArray<FString> ShelvedFiles;
+
+	/** Files that were moved to a new changelist if shelving from the Default Changelist */
+	TArray<FString> MovedFiles;
+
+	/** Changelist description if needed */
+	FString ChangelistDescription;
+
+	/** Changelist(s) to be updated */
+	FPlasticSourceControlChangelist InChangelistToUpdate;
+	FPlasticSourceControlChangelist OutChangelistToUpdate;
+};
+
+class FPlasticUnshelveWorker final : public IPlasticSourceControlWorker
+{
+public:
+	explicit FPlasticUnshelveWorker(FPlasticSourceControlProvider& InSourceControlProvider)
+		: IPlasticSourceControlWorker(InSourceControlProvider)
+	{}
+	virtual ~FPlasticUnshelveWorker() = default;
+
+	// IPlasticSourceControlWorker interface
+	virtual FName GetName() const override;
+	virtual bool Execute(FPlasticSourceControlCommand& InCommand) override;
+	virtual bool UpdateStates() override;
+
+protected:
+	/** List of files states after the unshelve */
+	TArray<FPlasticSourceControlState> States;
+
+	/** Changelist to be updated */
+	FPlasticSourceControlChangelist ChangelistToUpdate;
+};
+
+class FPlasticDeleteShelveWorker final : public IPlasticSourceControlWorker
+{
+public:
+	explicit FPlasticDeleteShelveWorker(FPlasticSourceControlProvider& InSourceControlProvider)
+		: IPlasticSourceControlWorker(InSourceControlProvider)
+	{}
+	virtual ~FPlasticDeleteShelveWorker() = default;
+
+	// IPlasticSourceControlWorker interface
+	virtual FName GetName() const override;
+	virtual bool Execute(FPlasticSourceControlCommand& InCommand) override;
+	virtual bool UpdateStates() override;
+
+protected:
+	/** List of files to remove from shelved files in changelist state */
+	TArray<FString> FilesToRemove;
+
+	/** Changelist to be updated */
+	FPlasticSourceControlChangelist ChangelistToUpdate;
+
+	/** Id of the new shelve (if only a selection of files are deleted from the shelve) */
+	int32 ShelveId = ISourceControlState::INVALID_REVISION;
+};
+
diff --git a/Engine/Plugins/Developer/PlasticSourceControl/Source/PlasticSourceControl/Private/PlasticSourceControlProjectSettings.h b/Engine/Plugins/Developer/PlasticSourceControl/Source/PlasticSourceControl/Private/PlasticSourceControlProjectSettings.h
--- a/Engine/Plugins/Developer/PlasticSourceControl/Source/PlasticSourceControl/Private/PlasticSourceControlProjectSettings.h
+++ b/Engine/Plugins/Developer/PlasticSourceControl/Source/PlasticSourceControl/Private/PlasticSourceControlProjectSettings.h
@@ -24,4 +24,8 @@ public:
 	/** If enabled, you'll be prompted to check out changed files (enabled by default). Checkout is needed to work with Changelists. */
 	UPROPERTY(config, EditAnywhere, Category = "Plastic SCM")
 	bool bPromptForCheckoutOnChange = true;
+
+	/** If a non-null value is set, limit the maximum number of revisions requested to Plastic SCM to display in the "History" window. */
+	UPROPERTY(config, EditAnywhere, Category = "Plastic SCM", meta=(ClampMin=0))
+	int32 LimitNumberOfRevisionsInHistory = 50;
 };
diff --git a/Engine/Plugins/Developer/PlasticSourceControl/Source/PlasticSourceControl/Private/PlasticSourceControlProvider.cpp b/Engine/Plugins/Developer/PlasticSourceControl/Source/PlasticSourceControl/Private/PlasticSourceControlProvider.cpp
--- a/Engine/Plugins/Developer/PlasticSourceControl/Source/PlasticSourceControl/Private/PlasticSourceControlProvider.cpp
+++ b/Engine/Plugins/Developer/PlasticSourceControl/Source/PlasticSourceControl/Private/PlasticSourceControlProvider.cpp
@@ -7,6 +7,7 @@
 #include "PlasticSourceControlOperations.h"
 #include "PlasticSourceControlProjectSettings.h"
 #include "PlasticSourceControlSettings.h"
+#include "PlasticSourceControlShell.h"
 #include "PlasticSourceControlState.h"
 #include "PlasticSourceControlUtils.h"
 #include "SPlasticSourceControlSettings.h"
@@ -58,9 +59,13 @@ void FPlasticSourceControlProvider::Init(bool bForceConnection)
 	{
 		// Execute a 'checkconnection' command to set bServerAvailable based on the connectivity of the server
 		TArray<FString> InfoMessages, ErrorMessages;
-		const bool bCommandSuccessful = PlasticSourceControlUtils::RunCommand(TEXT("checkconnection"), TArray<FString>(), TArray<FString>(), EConcurrency::Synchronous, InfoMessages, ErrorMessages);
-		bServerAvailable = bCommandSuccessful;
-		if (!bCommandSuccessful)
+		TArray<FString> Parameters;
+		if (PlasticSourceControlUtils::GetWorkspaceInfo(BranchName, RepositoryName, ServerUrl, ErrorMessages))
+		{
+			Parameters.Add(FString::Printf(TEXT("--server=%s"), *ServerUrl));
+		}
+		bServerAvailable = PlasticSourceControlUtils::RunCommand(TEXT("checkconnection"), Parameters, TArray<FString>(), InfoMessages, ErrorMessages);
+		if (!bServerAvailable)
 		{
 			FMessageLog SourceControlLog("SourceControl");
 			for (const FString& ErrorMessage : ErrorMessages)
@@ -90,10 +95,9 @@ void FPlasticSourceControlProvider::CheckPlasticAvailability()
 	{
 		// Find the path to the root Plastic directory (if any, else uses the ProjectDir)
 		const FString PathToProjectDir = FPaths::ConvertRelativePathToFull(FPaths::ProjectDir());
-		bWorkspaceFound = PlasticSourceControlUtils::FindRootDirectory(PathToProjectDir, PathToWorkspaceRoot);
 
 		// Launch the Plastic SCM cli shell on the background to issue all commands during this session
-		bPlasticAvailable = PlasticSourceControlUtils::LaunchBackgroundPlasticShell(PathToPlasticBinary, PathToWorkspaceRoot);
+		bPlasticAvailable = PlasticSourceControlShell::Launch(PathToPlasticBinary, PathToProjectDir);
 		if (!bPlasticAvailable)
 		{
 			return;
@@ -105,6 +109,14 @@ void FPlasticSourceControlProvider::CheckPlasticAvailability()
 			return;
 		}
 
+		FString ActualPathToPlasticBinary;
+		PlasticSourceControlUtils::GetCmLocation(ActualPathToPlasticBinary);
+
+		bWorkspaceFound = PlasticSourceControlUtils::GetWorkspacePath(PathToProjectDir, PathToWorkspaceRoot);
+
+
+		bUsesLocalReadOnlyState = PlasticSourceControlUtils::GetConfigSetFilesAsReadOnly();
+
 		// Get user name (from the global Plastic SCM client config)
 		PlasticSourceControlUtils::GetUserName(UserName);
 
@@ -126,7 +138,7 @@ void FPlasticSourceControlProvider::Close()
 	// clear the cache
 	StateCache.Empty();
 	// terminate the background 'cm shell' process and associated pipes
-	PlasticSourceControlUtils::Terminate();
+	PlasticSourceControlShell::Terminate();
 	// Remove all extensions to the "Source Control" menu in the Editor Toolbar
 	PlasticSourceControlMenu.Unregister();
 	// Unregister Console Commands
@@ -181,13 +193,13 @@ FText FPlasticSourceControlProvider::GetStatusText() const
 	Args.Add(TEXT("WorkspaceName"), FText::FromString(WorkspaceName));
 	Args.Add(TEXT("BranchName"), FText::FromString(BranchName));
 	// Detect special case for a partial checkout (CS:-1 in Gluon mode)!
-	if (-1 != ChangesetNumber)
+	if (IsPartialWorkspace())
 	{
-		Args.Add(TEXT("ChangesetNumber"), FText::FromString(FString::Printf(TEXT("%d  (regular full workspace)"), ChangesetNumber)));
+		Args.Add(TEXT("ChangesetNumber"), FText::FromString(FString::Printf(TEXT("N/A  (Gluon partial workspace)"))));
 	}
 	else
 	{
-		Args.Add(TEXT("ChangesetNumber"), FText::FromString(FString::Printf(TEXT("N/A  (Gluon/partial workspace)"))));
+		Args.Add(TEXT("ChangesetNumber"), FText::FromString(FString::Printf(TEXT("%d  (regular full workspace)"), ChangesetNumber)));
 	}
 	Args.Add(TEXT("UserName"), FText::FromString(UserName));
 	const FString DisplayName = PlasticSourceControlUtils::UserNameToDisplayName(UserName);
@@ -310,7 +322,6 @@ TArray<FSourceControlChangelistStateRef> FPlasticSourceControlProvider::GetCache
 	return Result;
 }
 
-
 TArray<FSourceControlStateRef> FPlasticSourceControlProvider::GetCachedStateByPredicate(TFunctionRef<bool(const FSourceControlStateRef&)> Predicate) const
 {
 	TArray<FSourceControlStateRef> Result;
@@ -401,7 +412,7 @@ void FPlasticSourceControlProvider::CancelOperation(const FSourceControlOperatio
 
 bool FPlasticSourceControlProvider::UsesLocalReadOnlyState() const
 {
-	return false; // TODO: use configuration
+	return bUsesLocalReadOnlyState;
 }
 
 bool FPlasticSourceControlProvider::UsesChangelists() const
@@ -431,12 +442,12 @@ bool FPlasticSourceControlProvider::AllowsDiffAgainstDepot() const
 
 TOptional<bool> FPlasticSourceControlProvider::IsAtLatestRevision() const
 {
-	return TOptional<bool>();
+	return TOptional<bool>(); // NOTE: used by code in UE5's Status Bar but currently dormant as far as I can tell
 }
 
 TOptional<int> FPlasticSourceControlProvider::GetNumLocalChanges() const
 {
-	return TOptional<int>();
+	return TOptional<int>(); // NOTE: used by code in UE5's Status Bar but currently dormant as far as I can tell
 }
 
 TSharedPtr<IPlasticSourceControlWorker, ESPMode::ThreadSafe> FPlasticSourceControlProvider::CreateWorker(const FName& InOperationName)
@@ -561,10 +572,14 @@ void FPlasticSourceControlProvider::Tick()
 			// dump any messages to output log
 			OutputCommandMessages(Command);
 
-			if (Command.Files.Num())
+			if (Command.Files.Num() > 1)
 			{
 				UE_LOG(LogSourceControl, Log, TEXT("%s of %d files processed in %.3lfs"), *Command.Operation->GetName().ToString(), Command.Files.Num(), (FPlatformTime::Seconds() - Command.StartTimestamp));
 			}
+			else if (Command.Files.Num() == 1)
+			{
+				UE_LOG(LogSourceControl, Log, TEXT("%s of %s processed in %.3lfs"), *Command.Operation->GetName().ToString(), *Command.Files[0], (FPlatformTime::Seconds() - Command.StartTimestamp));
+			}
 			else
 			{
 				UE_LOG(LogSourceControl, Log, TEXT("%s processed in %.3lfs"), *Command.Operation->GetName().ToString(), (FPlatformTime::Seconds() - Command.StartTimestamp));
@@ -600,7 +615,7 @@ TArray<TSharedRef<ISourceControlLabel>> FPlasticSourceControlProvider::GetLabels
 	return Tags;
 }
 
-TArray<FSourceControlChangelistRef> FPlasticSourceControlProvider::GetChangelists( EStateCacheUsage::Type InStateCacheUsage )
+TArray<FSourceControlChangelistRef> FPlasticSourceControlProvider::GetChangelists(EStateCacheUsage::Type InStateCacheUsage)
 {
 	if (!IsEnabled())
 	{
diff --git a/Engine/Plugins/Developer/PlasticSourceControl/Source/PlasticSourceControl/Private/PlasticSourceControlProvider.h b/Engine/Plugins/Developer/PlasticSourceControl/Source/PlasticSourceControl/Private/PlasticSourceControlProvider.h
--- a/Engine/Plugins/Developer/PlasticSourceControl/Source/PlasticSourceControl/Private/PlasticSourceControlProvider.h
+++ b/Engine/Plugins/Developer/PlasticSourceControl/Source/PlasticSourceControl/Private/PlasticSourceControlProvider.h
@@ -31,15 +31,15 @@ public:
 	virtual bool QueryStateBranchConfig(const FString& ConfigSrc, const FString& ConfigDest) override { return false; }
 	virtual void RegisterStateBranches(const TArray<FString>& BranchNames, const FString& ContentRoot) override {}
 	virtual int32 GetStateBranchIndex(const FString& InBranchName) const override { return INDEX_NONE; }
-	virtual ECommandResult::Type GetState( const TArray<FString>& InFiles, TArray<FSourceControlStateRef>& OutState, EStateCacheUsage::Type InStateCacheUsage ) override;
+	virtual ECommandResult::Type GetState(const TArray<FString>& InFiles, TArray<FSourceControlStateRef>& OutState, EStateCacheUsage::Type InStateCacheUsage) override;
 	virtual ECommandResult::Type GetState(const TArray<FSourceControlChangelistRef>& InChangelists, TArray<FSourceControlChangelistStateRef>& OutState, EStateCacheUsage::Type InStateCacheUsage) override;
 	virtual TArray<FSourceControlStateRef> GetCachedStateByPredicate(TFunctionRef<bool(const FSourceControlStateRef&)> Predicate) const override;
-	virtual FDelegateHandle RegisterSourceControlStateChanged_Handle( const FSourceControlStateChanged::FDelegate& SourceControlStateChanged ) override;
-	virtual void UnregisterSourceControlStateChanged_Handle( FDelegateHandle Handle ) override;
-	virtual ECommandResult::Type Execute( const FSourceControlOperationRef& InOperation, FSourceControlChangelistPtr InChangelist, const TArray<FString>& InFiles, EConcurrency::Type InConcurrency = EConcurrency::Synchronous, const FSourceControlOperationComplete& InOperationCompleteDelegate = FSourceControlOperationComplete() ) override;
+	virtual FDelegateHandle RegisterSourceControlStateChanged_Handle(const FSourceControlStateChanged::FDelegate& SourceControlStateChanged) override;
+	virtual void UnregisterSourceControlStateChanged_Handle(FDelegateHandle Handle) override;
+	virtual ECommandResult::Type Execute(const FSourceControlOperationRef& InOperation, FSourceControlChangelistPtr InChangelist, const TArray<FString>& InFiles, EConcurrency::Type InConcurrency = EConcurrency::Synchronous, const FSourceControlOperationComplete& InOperationCompleteDelegate = FSourceControlOperationComplete() ) override;
 	virtual bool CanExecuteOperation( const FSourceControlOperationRef& InOperation ) const override;
-	virtual bool CanCancelOperation( const FSourceControlOperationRef& InOperation ) const override;
-	virtual void CancelOperation( const FSourceControlOperationRef& InOperation ) override;
+	virtual bool CanCancelOperation(const FSourceControlOperationRef& InOperation) const override;
+	virtual void CancelOperation(const FSourceControlOperationRef& InOperation) override;
 	virtual bool UsesLocalReadOnlyState() const override;
 	virtual bool UsesChangelists() const override;
 	virtual bool UsesUncontrolledChangelists() const override;
@@ -49,8 +49,8 @@ public:
 	virtual TOptional<bool> IsAtLatestRevision() const override;
 	virtual TOptional<int> GetNumLocalChanges() const override;
 	virtual void Tick() override;
-	virtual TArray< TSharedRef<class ISourceControlLabel> > GetLabels( const FString& InMatchingSpec ) const override;
-	virtual TArray<FSourceControlChangelistRef> GetChangelists( EStateCacheUsage::Type InStateCacheUsage ) override;
+	virtual TArray<TSharedRef<class ISourceControlLabel>> GetLabels(const FString& InMatchingSpec) const override;
+	virtual TArray<FSourceControlChangelistRef> GetChangelists(EStateCacheUsage::Type InStateCacheUsage) override;
 #if SOURCE_CONTROL_WITH_SLATE
 	virtual TSharedRef<class SWidget> MakeSettingsWidget() const override;
 #endif
@@ -116,6 +116,12 @@ public:
 		return ChangesetNumber;
 	}
 
+	/** A partial/Gluon workspace doesn't match with a single changeset, which is identified by -1 */
+	inline bool IsPartialWorkspace() const
+	{
+		return (ChangesetNumber == -1);
+	}
+
 	/** Version of the Plastic SCM executable used */
 	inline const FSoftwareVersion& GetPlasticScmVersion() const
 	{
@@ -144,7 +150,7 @@ public:
 	 * Register a worker with the provider.
 	 * This is used internally so the provider can maintain a map of all available operations.
 	 */
-	void RegisterWorker( const FName& InName, const FGetPlasticSourceControlWorker& InDelegate );
+	void RegisterWorker(const FName& InName, const FGetPlasticSourceControlWorker& InDelegate);
 
 	/** Remove a named file from the state cache */
 	bool RemoveFileFromCache(const FString& Filename);
@@ -181,6 +187,9 @@ private:
 	/** Indicates if source control integration is available or not. */
 	bool bServerAvailable = false;
 
+	/** Whether Plastic SCM is configured to uses local read-only state to signal whether a file is editable ("SetFilesAsReadOnly" in client.conf) */
+	bool bUsesLocalReadOnlyState = false;
+
 	/** Critical section for thread safety of error messages that occurred after last Plastic command */
 	mutable FCriticalSection LastErrorsCriticalSection;
 
diff --git a/Engine/Plugins/Developer/PlasticSourceControl/Source/PlasticSourceControl/Private/PlasticSourceControlRevision.cpp b/Engine/Plugins/Developer/PlasticSourceControl/Source/PlasticSourceControl/Private/PlasticSourceControlRevision.cpp
--- a/Engine/Plugins/Developer/PlasticSourceControl/Source/PlasticSourceControl/Private/PlasticSourceControlRevision.cpp
+++ b/Engine/Plugins/Developer/PlasticSourceControl/Source/PlasticSourceControl/Private/PlasticSourceControlRevision.cpp
@@ -13,7 +13,7 @@
 
 #define LOCTEXT_NAMESPACE "PlasticSourceControl"
 
-bool FPlasticSourceControlRevision::Get(FString& InOutFilename, EConcurrency::Type InConcurrency) const
+bool FPlasticSourceControlRevision::Get(FString& InOutFilename, EConcurrency::Type InConcurrency /* = EConcurrency::Synchronous */) const
 {
 	if (InConcurrency != EConcurrency::Synchronous)
 	{
@@ -26,7 +26,15 @@ bool FPlasticSourceControlRevision::Get(FString& InOutFilename, EConcurrency::Ty
 		// create the diff dir if we don't already have it (Plastic wont)
 		IFileManager::Get().MakeDirectory(*FPaths::DiffDir(), true);
 		// create a unique temp file name based on the unique revision Id
-		const FString TempFileName = FString::Printf(TEXT("%stemp-%d-%s"), *FPaths::DiffDir(), ChangesetNumber, *FPaths::GetCleanFilename(Filename));
+		FString TempFileName;
+		if (ShelveId != ISourceControlState::INVALID_REVISION)
+		{
+			TempFileName = FString::Printf(TEXT("%stemp-sh%d-%s"), *FPaths::DiffDir(), ShelveId, *FPaths::GetCleanFilename(Filename));
+		}
+		else
+		{
+			TempFileName = FString::Printf(TEXT("%stemp-cs%d-%s"), *FPaths::DiffDir(), ChangesetNumber, *FPaths::GetCleanFilename(Filename));
+		}
 		InOutFilename = FPaths::ConvertRelativePathToFull(TempFileName);
 	}
 
@@ -39,8 +47,18 @@ bool FPlasticSourceControlRevision::Get(FString& InOutFilename, EConcurrency::Ty
 	{
 		const FString& PathToPlasticBinary = FPlasticSourceControlModule::Get().GetProvider().AccessSettings().GetBinaryPath();
 
-		// Format the revision specification of the file, like rev:Content/BP.uasset#cs:12@repo@server:8087
-		const FString RevisionSpecification = FString::Printf(TEXT("rev:%s#cs:%d@%s"), *Filename, ChangesetNumber, *State->RepSpec);
+		FString RevisionSpecification;
+		if (ShelveId != ISourceControlState::INVALID_REVISION)
+		{
+			// Format the revision specification of the shelved file, like rev:Content/BP.uasset#sh:33
+			// Note: the plugin doesn't support shelves on Xlinks (no known RepSpec)
+			RevisionSpecification = FString::Printf(TEXT("rev:%s#sh:%d"), *Filename, ShelveId);
+		}
+		else
+		{
+			// Format the revision specification of the checked-in file, like rev:Content/BP.uasset#cs:12@repo@server:8087
+			RevisionSpecification = FString::Printf(TEXT("rev:%s#cs:%d@%s"), *Filename, ChangesetNumber, *State->RepSpec);
+		}
 		bCommandSuccessful = PlasticSourceControlUtils::RunDumpToFile(PathToPlasticBinary, RevisionSpecification, InOutFilename);
 	}
 	else
diff --git a/Engine/Plugins/Developer/PlasticSourceControl/Source/PlasticSourceControl/Private/PlasticSourceControlRevision.h b/Engine/Plugins/Developer/PlasticSourceControl/Source/PlasticSourceControl/Private/PlasticSourceControlRevision.h
--- a/Engine/Plugins/Developer/PlasticSourceControl/Source/PlasticSourceControl/Private/PlasticSourceControlRevision.h
+++ b/Engine/Plugins/Developer/PlasticSourceControl/Source/PlasticSourceControl/Private/PlasticSourceControlRevision.h
@@ -51,6 +51,9 @@ public:
 	/** The revision to display to the user: use the changeset number */
 	FString Revision;
 
+	/** The Shelve ID instead of Changeset / Revision for case of shelved files */
+	int32 ShelveId = ISourceControlState::INVALID_REVISION;
+
 	/** The description of this revision */
 	FString Description;
 
diff --git a/Engine/Plugins/Developer/PlasticSourceControl/Source/PlasticSourceControl/Private/PlasticSourceControlShell.cpp b/Engine/Plugins/Developer/PlasticSourceControl/Source/PlasticSourceControl/Private/PlasticSourceControlShell.cpp
new file mode 100644
--- /dev/null
+++ b/Engine/Plugins/Developer/PlasticSourceControl/Source/PlasticSourceControl/Private/PlasticSourceControlShell.cpp
@@ -0,0 +1,348 @@
+// Copyright Epic Games, Inc. All Rights Reserved.
+
+#include "PlasticSourceControlShell.h"
+
+#include "PlasticSourceControlModule.h"
+#include "PlasticSourceControlProvider.h"
+
+#include "ISourceControlModule.h"
+
+#include "Misc/ScopeLock.h"
+#include "HAL/PlatformProcess.h"
+#include "HAL/PlatformTime.h"
+
+#include "Framework/Notifications/NotificationManager.h"
+#include "Widgets/Notifications/SNotificationList.h"
+
+#if PLATFORM_LINUX
+#include <sys/ioctl.h>
+#endif
+
+#if PLATFORM_WINDOWS
+#include "Windows/WindowsHWrapper.h" // SECURITY_ATTRIBUTES
+#undef GetUserName
+#endif
+
+
+#define LOCTEXT_NAMESPACE "PlasticSourceControl"
+
+
+namespace PlasticSourceControlShell
+{
+static const TCHAR* ShellCommandResultText = TEXT("CommandResult ");
+static const TCHAR* ShellUserInteractText = TEXT("Select your system [0-1]");
+
+// In/Out Pipes for the 'cm shell' persistent child process
+static void*			ShellOutputPipeRead = nullptr;
+static void*			ShellOutputPipeWrite = nullptr;
+static void*			ShellInputPipeRead = nullptr;
+static void*			ShellInputPipeWrite = nullptr;
+static FProcHandle		ShellProcessHandle;
+static FCriticalSection	ShellCriticalSection;
+static size_t			ShellCommandCounter = -1;
+static double			ShellCumulatedTime = 0.;
+
+// Internal function to cleanup (called under the critical section)
+static void _CleanupBackgroundCommandLineShell()
+{
+	FPlatformProcess::ClosePipe(ShellOutputPipeRead, ShellOutputPipeWrite);
+	FPlatformProcess::ClosePipe(ShellInputPipeRead, ShellInputPipeWrite);
+	ShellOutputPipeRead = ShellOutputPipeWrite = nullptr;
+	ShellInputPipeRead = ShellInputPipeWrite = nullptr;
+}
+
+// Internal function to launch the Plastic SCM background 'cm' process in interactive shell mode (called under the critical section)
+static bool _StartBackgroundPlasticShell(const FString& InPathToPlasticBinary, const FString& InWorkingDirectory)
+{
+	TRACE_CPUPROFILER_EVENT_SCOPE(PlasticSourceControlShell::_StartBackgroundPlasticShell);
+
+	const FString FullCommand(TEXT("shell --encoding=UTF-8"));
+
+	const bool bLaunchDetached = false;				// the new process will NOT have its own window
+	const bool bLaunchHidden = true;				// the new process will be minimized in the task bar
+	const bool bLaunchReallyHidden = bLaunchHidden; // the new process will not have a window or be in the task bar
+
+	const double StartTimestamp = FPlatformTime::Seconds();
+
+	verify(FPlatformProcess::CreatePipe(ShellOutputPipeRead, ShellOutputPipeWrite, false));	// For reading outputs from cm shell child process
+	verify(FPlatformProcess::CreatePipe(ShellInputPipeRead, ShellInputPipeWrite, true));	// For writing commands to cm shell child process
+
+#if PLATFORM_WINDOWS
+	ShellProcessHandle = FPlatformProcess::CreateProc(*InPathToPlasticBinary, *FullCommand, bLaunchDetached, bLaunchHidden, bLaunchReallyHidden, nullptr, 0, *InWorkingDirectory, ShellOutputPipeWrite, ShellInputPipeRead);
+#else // PLATFORM_LINUX
+	// Update working directory
+	char OriginalWorkingDirectory[PATH_MAX];
+	getcwd(OriginalWorkingDirectory, PATH_MAX);
+	chdir(TCHAR_TO_ANSI(*InWorkingDirectory));
+
+	ShellProcessHandle = FPlatformProcess::CreateProc(*InPathToPlasticBinary, *FullCommand, bLaunchDetached, bLaunchHidden, bLaunchReallyHidden, nullptr, 0, nullptr, ShellOutputPipeWrite, ShellInputPipeRead);
+
+	// Restore working directory
+	chdir(OriginalWorkingDirectory);
+#endif
+
+	if (!ShellProcessHandle.IsValid())
+	{
+		UE_LOG(LogSourceControl, Warning, TEXT("Failed to launch 'cm shell'")); // not a bug, just no Plastic SCM cli found
+		_CleanupBackgroundCommandLineShell();
+	}
+	else
+	{
+		const double ElapsedTime = (FPlatformTime::Seconds() - StartTimestamp);
+		UE_LOG(LogSourceControl, Verbose, TEXT("_StartBackgroundPlasticShell: '%s %s' ok (in %.3lfs, handle %d)"), *InPathToPlasticBinary, *FullCommand, ElapsedTime, ShellProcessHandle.Get());
+		ShellCommandCounter = 0;
+		ShellCumulatedTime = ElapsedTime;
+	}
+
+	return ShellProcessHandle.IsValid();
+}
+
+// Internal function (called under the critical section)
+// bInForceExit: set to true to immediately force close the process without trying to "exit" and wait for it
+static void _ExitBackgroundCommandLineShell(const bool bInForceExit = false)
+{
+	TRACE_CPUPROFILER_EVENT_SCOPE(PlasticSourceControlShell::_ExitBackgroundCommandLineShell);
+
+	if (ShellProcessHandle.IsValid())
+	{
+		if (FPlatformProcess::IsProcRunning(ShellProcessHandle))
+		{
+			if (bInForceExit)
+			{
+				FPlatformProcess::TerminateProc(ShellProcessHandle);
+			}
+			else
+			{
+				// Tell the 'cm shell' to exit
+				FPlatformProcess::WritePipe(ShellInputPipeWrite, TEXT("exit"));
+				// And wait up to one second for its termination
+				const double Timeout = 1.0;
+				const double StartTimestamp = FPlatformTime::Seconds();
+				while (FPlatformProcess::IsProcRunning(ShellProcessHandle))
+				{
+					if ((FPlatformTime::Seconds() - StartTimestamp) > Timeout)
+					{
+						UE_LOG(LogSourceControl, Warning, TEXT("ExitBackgroundCommandLineShell: cm shell didn't stop gracefully in %lfs."), Timeout);
+						FPlatformProcess::TerminateProc(ShellProcessHandle);
+						break;
+					}
+					FPlatformProcess::Sleep(0.01f);
+				}
+			}
+		}
+		FPlatformProcess::CloseProc(ShellProcessHandle);
+		_CleanupBackgroundCommandLineShell();
+	}
+}
+
+// Internal function (called under the critical section)
+// bInForceExit: set to true to immediately force close the process without trying to "exit" and wait for it
+static void _RestartBackgroundCommandLineShell(const bool bInForceExit = false)
+{
+	const FPlasticSourceControlProvider& Provider = FPlasticSourceControlModule::Get().GetProvider();
+	const FString& PathToPlasticBinary = Provider.AccessSettings().GetBinaryPath();
+	const FString& WorkingDirectory = Provider.GetPathToWorkspaceRoot();
+
+	_ExitBackgroundCommandLineShell(bInForceExit);
+	_StartBackgroundPlasticShell(PathToPlasticBinary, WorkingDirectory);
+}
+
+
+// Display a temporary failure notification in case of an error in the shell
+void DisplayFailureNotification(const FText& InNotificationText)
+{
+	FNotificationInfo* Info = new FNotificationInfo(InNotificationText);
+	Info->ExpireDuration = 10.0f;
+	FSlateNotificationManager::Get().QueueNotification(Info);
+	// NOTE: all source control operations run in a thread, so we cannot use MessageLog nor Notify() them since they can only be used from the Main/UI thread
+	// FMessageLog("SourceControl").Error(InNotificationText);
+	UE_LOG(LogSourceControl, Error, TEXT("%s"), *InNotificationText.ToString());
+}
+
+// Internal function (called under the critical section)
+static bool _RunCommandInternal(const FString& InCommand, const TArray<FString>& InParameters, const TArray<FString>& InFiles, FString& OutResults, FString& OutErrors)
+{
+	TRACE_CPUPROFILER_EVENT_SCOPE(PlasticSourceControlShell::_RunCommandInternal);
+
+	bool bResult = false;
+
+	ShellCommandCounter++;
+
+	// Detect previous crash of cm.exe and restart 'cm shell'
+	if (!FPlatformProcess::IsProcRunning(ShellProcessHandle))
+	{
+		UE_LOG(LogSourceControl, Warning, TEXT("RunCommand: 'cm shell' has stopped. Restarting!"));
+		_RestartBackgroundCommandLineShell();
+	}
+
+	// Start with the Plastic command itself ("status", "log", "checkin"...)
+	FString FullCommand = InCommand;
+	// Append to the command all parameters, and then finally the files
+	for (const FString& Parameter : InParameters)
+	{
+		FullCommand += TEXT(" ");
+		FullCommand += Parameter;
+	}
+	for (const FString& File : InFiles)
+	{
+		FullCommand += TEXT(" \"");
+		FullCommand += File;
+		FullCommand += TEXT("\"");
+	}
+	const FString LoggableCommand = FullCommand.Left(256); // Limit command log size to 256 characters
+	UE_LOG(LogSourceControl, Verbose, TEXT("RunCommand: '%s' (%d chars, %d files)"), *LoggableCommand, FullCommand.Len()+1, InFiles.Num());
+	FullCommand += TEXT('\n'); // Finalize the command line
+
+	// Send command to 'cm shell' process in UTF-8
+	// NOTE: this explicit conversion to UTF-8 shouldn't be needed since FPlatformProcess::WritePipe() says it does it, but reading the implementation for Windows Platform show it merily truncates 16bits to 8bits chars!
+	// NOTE: on the other hand, ReadPipe() does the conversion from UTF-8 correctly already!
+	const FTCHARToUTF8 FullCommandUtf8(*FullCommand);
+	const bool bWriteOk = FPlatformProcess::WritePipe(ShellInputPipeWrite, reinterpret_cast<const uint8*>(FullCommandUtf8.Get()), FullCommandUtf8.Length());
+
+	// And wait up to 180.0 seconds for any kind of output from cm shell: in case of lengthier operation, intermediate output (like percentage of progress) is expected, which would refresh the timeout
+	static const double Timeout = 180.0;
+	const double StartTimestamp = FPlatformTime::Seconds();
+	double LastActivity = StartTimestamp;
+	double LastLog = StartTimestamp;
+	static const double LogInterval = 5.0;
+	int32 PreviousLogLen = 0;
+	while (FPlatformProcess::IsProcRunning(ShellProcessHandle))
+	{
+		FString Output = FPlatformProcess::ReadPipe(ShellOutputPipeRead);
+		if (!Output.IsEmpty())
+		{
+			TRACE_CPUPROFILER_EVENT_SCOPE(PlasticSourceControlShell::_RunCommandInternal::ParseOutput);
+
+			LastActivity = FPlatformTime::Seconds(); // freshen the timestamp while cm is still actively outputting information
+			OutResults.Append(MoveTemp(Output));
+			// Search the output for the line containing the result code, also indicating the end of the command
+			const uint32 IndexCommandResult = OutResults.Find(ShellCommandResultText, ESearchCase::CaseSensitive, ESearchDir::FromEnd);
+			if (INDEX_NONE != IndexCommandResult)
+			{
+				const uint32 IndexEndResult = OutResults.Find(pchDelim, ESearchCase::CaseSensitive, ESearchDir::FromStart, IndexCommandResult + 14);
+				if (INDEX_NONE != IndexEndResult)
+				{
+					const FString Result = OutResults.Mid(IndexCommandResult + 14, IndexEndResult - IndexCommandResult - 14);
+					const int32 ResultCode = FCString::Atoi(*Result);
+					bResult = (ResultCode == 0);
+					// remove the CommandResult line from the OutResults
+					OutResults.RemoveAt(IndexCommandResult, OutResults.Len() - IndexCommandResult);
+					break;
+				}
+			}
+
+			// Search the output for a potential user interaction request (in case the authentication token isn't saved or valid anymore)
+			const uint32 IndexPrompt = OutResults.Find(ShellUserInteractText, ESearchCase::CaseSensitive, ESearchDir::FromEnd);
+			if (INDEX_NONE != IndexPrompt)
+			{
+				const FText ShellRequiresInteractionError(LOCTEXT("SourceControlShell_AskAuthenticate", "Plastic SCM command line requires user interaction.\nSign in using the Plastic SCM client."));
+				DisplayFailureNotification(ShellRequiresInteractionError);
+
+				// Restart the shell without waiting, it is forever blocked waiting for user input
+				_RestartBackgroundCommandLineShell(true);
+				break; // and quit the loop
+			}
+		}
+		else if ((FPlatformTime::Seconds() - LastLog > LogInterval) && (PreviousLogLen < OutResults.Len()))
+		{
+			// In case of long running operation, start to print intermediate output from cm shell (like percentage of progress)
+			UE_LOG(LogSourceControl, Log, TEXT("RunCommand: '%s' in progress for %.3lfs... (%d chars):\n%s"), *InCommand, (FPlatformTime::Seconds() - StartTimestamp), OutResults.Len() - PreviousLogLen, *OutResults.Mid(PreviousLogLen));
+			PreviousLogLen = OutResults.Len();
+			LastLog = FPlatformTime::Seconds(); // freshen the timestamp of last log
+		}
+		else if (FPlatformTime::Seconds() - LastActivity > Timeout)
+		{
+			// In case of timeout, ask the blocking 'cm shell' process to exit, detach from it and restart it immediately
+			UE_LOG(LogSourceControl, Error, TEXT("RunCommand: '%s' TIMEOUT after %.3lfs output (%d chars):\n%s"), *InCommand, (FPlatformTime::Seconds() - StartTimestamp), OutResults.Len(), *OutResults.Mid(PreviousLogLen));
+			_RestartBackgroundCommandLineShell(true);
+			// Return output results as error so they get propagated to the Message Log window
+			OutErrors = MoveTemp(OutResults);
+			return false;
+		}
+		else if (IsEngineExitRequested())
+		{
+			UE_LOG(LogSourceControl, Warning, TEXT("RunCommand: '%s' Engine Exit was requested after %.3lfs output (%d chars):\n%s"), *InCommand, (FPlatformTime::Seconds() - StartTimestamp), OutResults.Len() - PreviousLogLen, *OutResults.Mid(PreviousLogLen));
+			_ExitBackgroundCommandLineShell();
+		}
+
+		TRACE_CPUPROFILER_EVENT_SCOPE(PlasticSourceControlShell::_RunCommandInternal::Sleep);
+		FPlatformProcess::Sleep(0.001f);
+	}
+	const double ElapsedTime = (FPlatformTime::Seconds() - StartTimestamp);
+
+	if (!InCommand.Equals(TEXT("exit")))
+	{
+		if (!FPlatformProcess::IsProcRunning(ShellProcessHandle))
+		{
+			// 'cm shell' normally only terminates in case of 'exit' command. Will restart on next command.
+			UE_LOG(LogSourceControl, Error, TEXT("RunCommand: '%s' 'cm shell' stopped after %.3lfs output (%d chars):\n%s"), *LoggableCommand, ElapsedTime, OutResults.Len(), *OutResults.Left(4096)); // Limit result size to 4096 characters
+		}
+		else if (!bResult)
+		{
+			UE_LOG(LogSourceControl, Warning, TEXT("RunCommand: '%s' (in %.3lfs) output (%d chars):\n%s"), *LoggableCommand, ElapsedTime, OutResults.Len(), *OutResults.Left(4096)); // Limit result size to 4096 characters
+		}
+		else
+		{
+			if (PreviousLogLen > 0)
+			{
+				UE_LOG(LogSourceControl, Log, TEXT("RunCommand: '%s' (in %.3lfs) output (%d chars):\n%s"), *LoggableCommand, ElapsedTime, OutResults.Len(), *OutResults.Mid(PreviousLogLen).Left(4096)); // Limit result size to 4096 characters
+			}
+			else
+			{
+				if (OutResults.Len() <= 200) // Limit result size to 200 characters
+				{
+					UE_LOG(LogSourceControl, Log, TEXT("RunCommand: '%s' (in %.3lfs) output (%d chars):\n%s"), *LoggableCommand, ElapsedTime, OutResults.Len(), *OutResults);
+				}
+				else
+				{
+					UE_LOG(LogSourceControl, Log, TEXT("RunCommand: '%s' (in %.3lfs) (output %d chars not displayed)"), *LoggableCommand, ElapsedTime, OutResults.Len());
+					UE_LOG(LogSourceControl, Verbose, TEXT("\n%s"), *OutResults.Left(4096));; // Limit result size to 4096 characters
+				}
+			}
+		}
+	}
+	// Return output as error if result code is an error
+	if (!bResult)
+	{
+		OutErrors = MoveTemp(OutResults);
+	}
+
+	ShellCumulatedTime += ElapsedTime;
+	UE_LOG(LogSourceControl, Verbose, TEXT("RunCommand: cumulated time spent in shell: %.3lfs (count %d)"), ShellCumulatedTime, ShellCommandCounter);
+
+	return bResult;
+}
+
+// Launch the Plastic SCM 'cm shell' process in background for optimized successive commands (thread-safe)
+bool Launch(const FString& InPathToPlasticBinary, const FString& InWorkingDirectory)
+{
+	// Protect public APIs from multi-thread access
+	FScopeLock Lock(&ShellCriticalSection);
+
+	// terminate previous shell if one is already running
+	_ExitBackgroundCommandLineShell();
+
+	return _StartBackgroundPlasticShell(InPathToPlasticBinary, InWorkingDirectory);
+}
+
+// Terminate the background 'cm shell' process and associated pipes (thread-safe)
+void Terminate()
+{
+	// Protect public APIs from multi-thread access
+	FScopeLock Lock(&ShellCriticalSection);
+
+	_ExitBackgroundCommandLineShell();
+}
+
+// Run command and return the raw result
+bool RunCommand(const FString& InCommand, const TArray<FString>& InParameters, const TArray<FString>& InFiles, FString& OutResults, FString& OutErrors)
+{
+	// Protect public APIs from multi-thread access
+	FScopeLock Lock(&ShellCriticalSection);
+
+	return _RunCommandInternal(InCommand, InParameters, InFiles, OutResults, OutErrors);
+}
+
+} // namespace PlasticSourceControlShell
+
+#undef LOCTEXT_NAMESPACE
diff --git a/Engine/Plugins/Developer/PlasticSourceControl/Source/PlasticSourceControl/Private/PlasticSourceControlShell.h b/Engine/Plugins/Developer/PlasticSourceControl/Source/PlasticSourceControl/Private/PlasticSourceControlShell.h
new file mode 100644
--- /dev/null
+++ b/Engine/Plugins/Developer/PlasticSourceControl/Source/PlasticSourceControl/Private/PlasticSourceControlShell.h
@@ -0,0 +1,42 @@
+// Copyright Epic Games, Inc. All Rights Reserved.
+
+#pragma once
+
+#include "CoreMinimal.h"
+#include "ISourceControlProvider.h" // for EConcurrency
+
+namespace PlasticSourceControlShell
+{
+#if PLATFORM_WINDOWS
+	constexpr const TCHAR* pchDelim = TEXT("\r\n");
+#else
+	constexpr const TCHAR* pchDelim = TEXT("\n");
+#endif
+
+
+/**
+ * Launch the Plastic SCM "shell" command line process to run it in the background.
+ *
+ * @param	InPathToPlasticBinary	The path to the Plastic binary
+ * @param	InWorkspaceRoot			The workspace from where to run the command - usually the Game directory
+ * @returns true if the command succeeded and returned no errors
+ */
+bool Launch(const FString& InPathToPlasticBinary, const FString& InWorkspaceRoot);
+
+/** Terminate the background 'cm shell' process and associated pipes */
+void Terminate();
+
+
+/**
+ * Run a Plastic command - the result is the output of cm, as a multi-line string.
+ *
+ * @param	InCommand			The Plastic command - e.g. commit
+ * @param	InParameters		The parameters to the Plastic command
+ * @param	InFiles				The files to be operated on
+ * @param	OutResults			The results (from StdOut) as a multi-line string.
+ * @param	OutErrors			Any errors (from StdErr) as a multi-line string.
+ * @returns true if the command succeeded and returned no errors
+ */
+bool RunCommand(const FString& InCommand, const TArray<FString>& InParameters, const TArray<FString>& InFiles, FString& OutResults, FString& OutErrors);
+
+} // namespace PlasticSourceControlShell
diff --git a/Engine/Plugins/Developer/PlasticSourceControl/Source/PlasticSourceControl/Private/PlasticSourceControlState.cpp b/Engine/Plugins/Developer/PlasticSourceControl/Source/PlasticSourceControl/Private/PlasticSourceControlState.cpp
--- a/Engine/Plugins/Developer/PlasticSourceControl/Source/PlasticSourceControl/Private/PlasticSourceControlState.cpp
+++ b/Engine/Plugins/Developer/PlasticSourceControl/Source/PlasticSourceControl/Private/PlasticSourceControlState.cpp
@@ -8,32 +8,28 @@
 
 #define LOCTEXT_NAMESPACE "PlasticSourceControl.State"
 
-namespace EWorkspaceState
-{
-const TCHAR* ToString(EWorkspaceState::Type InWorkspaceState)
+const TCHAR* FPlasticSourceControlState::ToString() const
 {
 	const TCHAR* WorkspaceStateStr = nullptr;
-	switch (InWorkspaceState)
+	switch (WorkspaceState)
 	{
-	case EWorkspaceState::Unknown: WorkspaceStateStr = TEXT("Unknown"); break;
-	case EWorkspaceState::Ignored: WorkspaceStateStr = TEXT("Ignored"); break;
-	case EWorkspaceState::Controlled: WorkspaceStateStr = TEXT("Controlled"); break;
-	case EWorkspaceState::CheckedOut: WorkspaceStateStr = TEXT("CheckedOut"); break;
-	case EWorkspaceState::Added: WorkspaceStateStr = TEXT("Added"); break;
-	case EWorkspaceState::Moved: WorkspaceStateStr = TEXT("Moved"); break;
-	case EWorkspaceState::Copied: WorkspaceStateStr = TEXT("Copied"); break;
-	case EWorkspaceState::Replaced: WorkspaceStateStr = TEXT("Replaced"); break;
-	case EWorkspaceState::Deleted: WorkspaceStateStr = TEXT("Deleted"); break;
-	case EWorkspaceState::LocallyDeleted: WorkspaceStateStr = TEXT("LocallyDeleted"); break;
-	case EWorkspaceState::Changed: WorkspaceStateStr = TEXT("Changed"); break;
-	case EWorkspaceState::Conflicted: WorkspaceStateStr = TEXT("Conflicted"); break;
-	case EWorkspaceState::LockedByOther: WorkspaceStateStr = TEXT("LockedByOther"); break;
-	case EWorkspaceState::Private: WorkspaceStateStr = TEXT("Private"); break;
-	default: WorkspaceStateStr = TEXT("???"); break;
+		case EWorkspaceState::Unknown: WorkspaceStateStr = TEXT("Unknown"); break;
+		case EWorkspaceState::Ignored: WorkspaceStateStr = TEXT("Ignored"); break;
+		case EWorkspaceState::Controlled: WorkspaceStateStr = TEXT("Controlled"); break;
+		case EWorkspaceState::CheckedOut: WorkspaceStateStr = TEXT("CheckedOut"); break;
+		case EWorkspaceState::Added: WorkspaceStateStr = TEXT("Added"); break;
+		case EWorkspaceState::Moved: WorkspaceStateStr = TEXT("Moved"); break;
+		case EWorkspaceState::Copied: WorkspaceStateStr = TEXT("Copied"); break;
+		case EWorkspaceState::Replaced: WorkspaceStateStr = TEXT("Replaced"); break;
+		case EWorkspaceState::Deleted: WorkspaceStateStr = TEXT("Deleted"); break;
+		case EWorkspaceState::LocallyDeleted: WorkspaceStateStr = TEXT("LocallyDeleted"); break;
+		case EWorkspaceState::Changed: WorkspaceStateStr = TEXT("Changed"); break;
+		case EWorkspaceState::Conflicted: WorkspaceStateStr = TEXT("Conflicted"); break;
+		case EWorkspaceState::Private: WorkspaceStateStr = TEXT("Private"); break;
+		default: WorkspaceStateStr = TEXT("???"); break;
 	}
 	return WorkspaceStateStr;
 }
-} // namespace EWorkspaceState
 
 
 int32 FPlasticSourceControlState::GetHistorySize() const
@@ -77,7 +73,7 @@ TSharedPtr<class ISourceControlRevision, ESPMode::ThreadSafe> FPlasticSourceCont
 {
 	for (const auto& Revision : History)
 	{
-		// look for the SHA1 id of the file, not the commit id (revision)
+		// look for the changeset number, not the revision
 		if (Revision->ChangesetNumber == PendingMergeBaseChangeset)
 		{
 			return Revision;
@@ -91,7 +87,7 @@ TSharedPtr<class ISourceControlRevision, ESPMode::ThreadSafe> FPlasticSourceCont
 {
 	for (const auto& Revision : History)
 	{
-		// look for the SHA1 id of the file, not the commit id (revision)
+		// look for the changeset number, not the revision
 		if (Revision->ChangesetNumber == LocalRevisionChangeset)
 		{
 			return Revision;
@@ -107,9 +103,13 @@ FSlateIcon FPlasticSourceControlState::GetIcon() const
 	{
 		return FSlateIcon(FRevisionControlStyleManager::GetStyleSetName(), "RevisionControl.NotAtHeadRevision");
 	}
-	else if (WorkspaceState != EWorkspaceState::CheckedOut && WorkspaceState != EWorkspaceState::LockedByOther)
+	else if (WorkspaceState != EWorkspaceState::CheckedOut && WorkspaceState != EWorkspaceState::Conflicted)
 	{
-		if (IsModifiedInOtherBranch())
+		if (IsCheckedOutOther())
+		{
+			return FSlateIcon(FRevisionControlStyleManager::GetStyleSetName(), "RevisionControl.CheckedOutByOtherUser", NAME_None, "RevisionControl.CheckedOutByOtherUserBadge");
+		}
+		else if (IsModifiedInOtherBranch())
 		{
 			return FSlateIcon(FRevisionControlStyleManager::GetStyleSetName(), "RevisionControl.ModifiedOtherBranch", NAME_None, "RevisionControl.ModifiedBadge");
 		}
@@ -132,9 +132,7 @@ FSlateIcon FPlasticSourceControlState::GetIcon() const
 	case EWorkspaceState::LocallyDeleted:
 		return FSlateIcon(FAppStyle::GetAppStyleSetName(), "Plastic.LocallyDeleted"); // custom
 	case EWorkspaceState::Conflicted:
-		return FSlateIcon(FAppStyle::GetAppStyleSetName(), "Plastic.Conflicted"); // custom
-	case EWorkspaceState::LockedByOther:
-		return FSlateIcon(FRevisionControlStyleManager::GetStyleSetName(), "RevisionControl.CheckedOutByOtherUser", NAME_None, "RevisionControl.CheckedOutByOtherUserBadge");
+		return FSlateIcon(FRevisionControlStyleManager::GetStyleSetName(), "RevisionControl.Conflicted");
 	case EWorkspaceState::Private: // Not controlled
 		return FSlateIcon(FRevisionControlStyleManager::GetStyleSetName(), "RevisionControl.NotInDepot");
 	case EWorkspaceState::Ignored:
@@ -148,19 +146,26 @@ FSlateIcon FPlasticSourceControlState::GetIcon() const
 
 FText FPlasticSourceControlState::GetDisplayName() const
 {
+	FNumberFormattingOptions NoCommas;
+	NoCommas.UseGrouping = false;
 	if (!IsCurrent())
 	{
-		return LOCTEXT("NotCurrent", "Not current");
+		return FText::Format(LOCTEXT("NotCurrent", "Not at the head revision CS:{0} {1} (local revision is CS:{2})"),
+			FText::AsNumber(DepotRevisionChangeset), FText::FromString(HeadUserName), FText::AsNumber(LocalRevisionChangeset, &NoCommas));
 	}
-	else if (WorkspaceState != EWorkspaceState::LockedByOther)
+	else if (WorkspaceState != EWorkspaceState::CheckedOut && WorkspaceState != EWorkspaceState::Conflicted)
 	{
-		if (IsModifiedInOtherBranch())
+		if (IsCheckedOutOther())
+		{
+			return FText::Format(LOCTEXT("CheckedOutOther", "Checked out by {0} in {1}"), FText::FromString(LockedBy), FText::FromString(LockedWhere));
+		}
+		else if (IsModifiedInOtherBranch())
 		{
-			FNumberFormattingOptions NoCommas;
-			NoCommas.UseGrouping = false;
-			return FText::Format(LOCTEXT("ModifiedOtherBranch", "Modified in {0} CS:{1} by {2}"), FText::FromString(HeadBranch), FText::AsNumber(HeadChangeList, &NoCommas), FText::FromString(HeadUserName));
+			return FText::Format(LOCTEXT("ModifiedOtherBranch", "Modified in {0} as CS:{1} by {2} (local revision is CS:{3})"),
+				FText::FromString(HeadBranch), FText::AsNumber(HeadChangeList, &NoCommas), FText::FromString(HeadUserName), FText::AsNumber(LocalRevisionChangeset, &NoCommas));
 		}
 	}
+
 	switch (WorkspaceState)
 	{
 	case EWorkspaceState::Unknown:
@@ -186,9 +191,7 @@ FText FPlasticSourceControlState::GetDisplayName() const
 	case EWorkspaceState::Changed:
 		return LOCTEXT("Changed", "Changed");
 	case EWorkspaceState::Conflicted:
-		return LOCTEXT("ContentsConflict", "Conflicted");
-	case EWorkspaceState::LockedByOther:
-		return FText::Format(LOCTEXT("CheckedOutOther", "Checked out by: {0} in {1}"), FText::FromString(LockedBy), FText::FromString(LockedWhere));
+		return LOCTEXT("Conflicted", "Conflicted");
 	case EWorkspaceState::Private:
 		return LOCTEXT("NotControlled", "Not Under Source Control");
 	}
@@ -198,17 +201,23 @@ FText FPlasticSourceControlState::GetDisplayName() const
 
 FText FPlasticSourceControlState::GetDisplayTooltip() const
 {
+	FNumberFormattingOptions NoCommas;
+	NoCommas.UseGrouping = false;
 	if (!IsCurrent())
 	{
-		return FText::Format(LOCTEXT("NotCurrent_Tooltip", "Not at the head revision CS:{0} by {1}"), FText::AsNumber(DepotRevisionChangeset), FText::FromString(HeadUserName));
+		return FText::Format(LOCTEXT("NotCurrent_Tooltip", "Not at the head revision CS:{0} {1} (local revision is CS:{2})"),
+			FText::AsNumber(DepotRevisionChangeset), FText::FromString(HeadUserName), FText::AsNumber(LocalRevisionChangeset, &NoCommas));
 	}
-	else if (WorkspaceState != EWorkspaceState::LockedByOther)
+	else if (WorkspaceState != EWorkspaceState::CheckedOut && WorkspaceState != EWorkspaceState::Conflicted)
 	{
-		if (IsModifiedInOtherBranch())
+		if (IsCheckedOutOther())
+		{
+			return FText::Format(LOCTEXT("CheckedOutOther_Tooltip", "Checked out by {0} in {1}"), FText::FromString(LockedBy), FText::FromString(LockedWhere));
+		}
+		else if (IsModifiedInOtherBranch())
 		{
-			FNumberFormattingOptions NoCommas;
-			NoCommas.UseGrouping = false;
-			return FText::Format(LOCTEXT("ModifiedOtherBranch_Tooltip", "Modified in {0} CS:{1} by {2}"), FText::FromString(HeadBranch), FText::AsNumber(HeadChangeList, &NoCommas), FText::FromString(HeadUserName));
+			return FText::Format(LOCTEXT("ModifiedOtherBranch_Tooltip", "Modified in {0} as CS:{1} by {2} (local revision is CS:{3})"),
+				FText::FromString(HeadBranch), FText::AsNumber(HeadChangeList, &NoCommas), FText::FromString(HeadUserName), FText::AsNumber(LocalRevisionChangeset, &NoCommas));
 		}
 	}
 
@@ -225,7 +234,12 @@ FText FPlasticSourceControlState::GetDisplayTooltip() const
 	case EWorkspaceState::Added:
 		return LOCTEXT("Added_Tooltip", "Added");
 	case EWorkspaceState::Moved:
-		return LOCTEXT("Moved_Tooltip", "Moved or renamed");
+	{
+		FString MoveOrigin = MovedFrom;
+		FPaths::MakePathRelativeTo(MoveOrigin, *LocalFilename);
+		return FText::Format(LOCTEXT("Moved_Tooltip", "Moved from {0}"),
+			FText::FromString(FPaths::GetBaseFilename(MoveOrigin, false)));
+	}
 	case EWorkspaceState::Copied:
 		return LOCTEXT("Copied_Tooltip", "Copied");
 	case EWorkspaceState::Replaced:
@@ -237,9 +251,8 @@ FText FPlasticSourceControlState::GetDisplayTooltip() const
 	case EWorkspaceState::Changed:
 		return LOCTEXT("Modified_Tooltip", "Locally modified");
 	case EWorkspaceState::Conflicted:
-		return LOCTEXT("ContentsConflict_Tooltip", "Conflict with updates received from the repository");
-	case EWorkspaceState::LockedByOther:
-		return FText::Format(LOCTEXT("CheckedOutOther_Tooltip", "Checked out by {0} in {1}"), FText::FromString(LockedBy), FText::FromString(LockedWhere));
+		return FText::Format(LOCTEXT("Conflicted_Tooltip", "Conflict merging from source/remote CS:{0} into target/local CS:{1})"),
+			FText::AsNumber(PendingMergeSourceChangeset, &NoCommas), FText::AsNumber(LocalRevisionChangeset, &NoCommas));
 	case EWorkspaceState::Private:
 		return LOCTEXT("NotControlled_Tooltip", "Private: not under version control");
 	}
@@ -261,19 +274,19 @@ const FDateTime& FPlasticSourceControlState::GetTimeStamp() const
 	return TimeStamp;
 }
 
-// Deleted and Missing assets cannot appear in the Content Browser but does appear in Submit to Source Control Window
 bool FPlasticSourceControlState::CanCheckIn() const
 {
-	const bool bCanCheckIn =
-		  (WorkspaceState == EWorkspaceState::Added
-		|| WorkspaceState == EWorkspaceState::Deleted
-		|| WorkspaceState == EWorkspaceState::LocallyDeleted
-		|| WorkspaceState == EWorkspaceState::Changed
-		|| WorkspaceState == EWorkspaceState::Moved
-		|| WorkspaceState == EWorkspaceState::Copied
-		|| WorkspaceState == EWorkspaceState::Replaced
-		|| WorkspaceState == EWorkspaceState::CheckedOut)
-		&& !IsConflicted() && IsCurrent();
+	// Deleted assets don't appear in the Content Browser but in Submit to Source Control Window
+	const bool bCanCheckIn =  (WorkspaceState == EWorkspaceState::Added
+							|| WorkspaceState == EWorkspaceState::Deleted
+							|| WorkspaceState == EWorkspaceState::LocallyDeleted
+							|| WorkspaceState == EWorkspaceState::Changed
+							|| WorkspaceState == EWorkspaceState::Moved
+							|| WorkspaceState == EWorkspaceState::Copied
+							|| WorkspaceState == EWorkspaceState::Replaced
+							|| WorkspaceState == EWorkspaceState::CheckedOut)
+							&& !IsCheckedOutOther()	// Is not already checked-out elsewhere
+							&& IsCurrent();			// Is up to date (at the revision of the repo)
 
 	if (!IsUnknown())
 	{
@@ -293,7 +306,8 @@ bool FPlasticSourceControlState::CanCheckout() const
 	const bool bCanCheckout  =    (WorkspaceState == EWorkspaceState::Controlled	// In source control, Unmodified
 								|| WorkspaceState == EWorkspaceState::Changed		// In source control, but not checked-out
 								|| WorkspaceState == EWorkspaceState::Replaced)		// In source control, merged, waiting for checkin to conclude the merge
-								&& IsCurrent(); // Is up to date (at the revision of the repo)
+								&& !IsCheckedOutOther()	// Is not already checked-out elsewhere
+								&& IsCurrent();			// Is up to date (at the revision of the repo)
 
 	if (!IsUnknown())
 	{
@@ -333,7 +347,15 @@ bool FPlasticSourceControlState::IsCheckedOutOther(FString* Who) const
 	{
 		*Who = LockedBy;
 	}
-	const bool bIsLockedByOther = WorkspaceState == EWorkspaceState::LockedByOther;
+
+	// If the asset is Locked but not CheckedOut locally, it means it is locked somewhere else
+	const bool bIsLocked = !LockedBy.IsEmpty();
+	const bool bIsCheckedOut = WorkspaceState == EWorkspaceState::CheckedOut
+							|| WorkspaceState == EWorkspaceState::Added
+							|| WorkspaceState == EWorkspaceState::Moved
+							|| WorkspaceState == EWorkspaceState::Conflicted	// In source control, waiting for merged
+							|| WorkspaceState == EWorkspaceState::Replaced;		// In source control, merged, waiting for checkin to conclude the merge
+	const bool bIsLockedByOther = bIsLocked && !bIsCheckedOut;
 
 	if (bIsLockedByOther)
 	{
@@ -385,9 +407,10 @@ bool FPlasticSourceControlState::IsCurrent() const
 
 bool FPlasticSourceControlState::IsSourceControlled() const
 {
+	// NOTE: the Editor Collections rely on the default 'Unknown' state (until the actual file status is obtained) to be considered "in source control"
 	const bool bIsSourceControlled = WorkspaceState != EWorkspaceState::Private
-								  && WorkspaceState != EWorkspaceState::Ignored
-								  && WorkspaceState != EWorkspaceState::Unknown;
+								  && WorkspaceState != EWorkspaceState::Ignored;
+								  // WorkspaceState != EWorkspaceState::Unknown 
 
 	if (!bIsSourceControlled && !IsUnknown()) UE_LOG(LogSourceControl, Log, TEXT("%s NOT SourceControlled"), *LocalFilename);
 
@@ -450,7 +473,7 @@ bool FPlasticSourceControlState::IsModified() const
 	//
 	// Thus, before checkin UE4 Editor call RevertUnchangedFiles() in PromptForCheckin() and CheckinFiles().
 	//
-	// So here we must take care to enumerate all states that need to be commited, all other will be discarded:
+	// So here we must take care to enumerate all states that need to be committed, all other will be discarded:
 	//  - Unknown
 	//  - Controlled (Unchanged)
 	//  - Private (Not Controlled)
@@ -480,7 +503,7 @@ bool FPlasticSourceControlState::CanAdd() const
 
 bool FPlasticSourceControlState::IsConflicted() const
 {
-	if (WorkspaceState == EWorkspaceState::Conflicted) UE_LOG(LogSourceControl, Log, TEXT("%s IsConflicted"), *LocalFilename);
+	if (WorkspaceState == EWorkspaceState::Conflicted) UE_LOG(LogSourceControl, Verbose, TEXT("%s IsConflicted"), *LocalFilename);
 
 	return WorkspaceState == EWorkspaceState::Conflicted;
 }
diff --git a/Engine/Plugins/Developer/PlasticSourceControl/Source/PlasticSourceControl/Private/PlasticSourceControlState.h b/Engine/Plugins/Developer/PlasticSourceControl/Source/PlasticSourceControl/Private/PlasticSourceControlState.h
--- a/Engine/Plugins/Developer/PlasticSourceControl/Source/PlasticSourceControl/Private/PlasticSourceControlState.h
+++ b/Engine/Plugins/Developer/PlasticSourceControl/Source/PlasticSourceControl/Private/PlasticSourceControlState.h
@@ -8,10 +8,8 @@
 #include "PlasticSourceControlRevision.h"
 #include "PlasticSourceControlChangelist.h"
 
-namespace EWorkspaceState
+enum class EWorkspaceState
 {
-	enum Type
-	{
 		Unknown,
 		Ignored,
 		Controlled, // called "Pristine" in Perforce, "Unchanged" in Git, "Clean" in SVN
@@ -24,13 +22,8 @@ namespace EWorkspaceState
 		LocallyDeleted, // Missing
 		Changed, // Modified but not CheckedOut
 		Conflicted,
-		LockedByOther, // LockedBy with name of someone else than cm whoami
 		Private, // "Not Controlled"/"Not In Depot"/"Untracked"
-	};
-
-	// debug log utility
-	const TCHAR* ToString(EWorkspaceState::Type InWorkspaceState);
-} // namespace EWorkspaceState
+};
 
 class FPlasticSourceControlState : public ISourceControlState
 {
@@ -40,6 +33,12 @@ public:
 	{
 	}
 
+	FPlasticSourceControlState(FString&& InLocalFilename, EWorkspaceState InWorkspaceState)
+		: LocalFilename(MoveTemp(InLocalFilename))
+		, WorkspaceState(InWorkspaceState)
+	{
+	}
+
 	FPlasticSourceControlState() = delete;
 	FPlasticSourceControlState(const FPlasticSourceControlState& InState) = delete;
 	const FPlasticSourceControlState& operator=(const FPlasticSourceControlState& InState) = delete;
@@ -65,6 +64,7 @@ public:
 		PendingMergeSourceChangeset = InState.PendingMergeSourceChangeset;
 		PendingMergeParameters = MoveTemp(InState.PendingMergeParameters);
 		// Update "fileinfo" information only if the command was issued
+		// Don't override "fileinfo" information in case of an optimized/lightweight "whole folder status" triggered by a global Submit Content or Refresh
 		if (InState.DepotRevisionChangeset != INVALID_REVISION)
 		{
 			LockedBy = MoveTemp(InState.LockedBy);
@@ -79,12 +79,6 @@ public:
 			HeadUserName = MoveTemp(InState.HeadUserName);
 			HeadModTime = MoveTemp(InState.HeadModTime);
 		}
-		// Don't override "fileinfo" information in case of an optimized/lightweight "whole folder status" triggered by a global Submit Content or Refresh
-		// and regenerate the LockedByOther state based on LockedBy
-		else if (!IsCheckedOut() && !LockedBy.IsEmpty())
-		{
-			WorkspaceState = EWorkspaceState::LockedByOther;
-		}
 		MovedFrom = MoveTemp(InState.MovedFrom);
 		TimeStamp = InState.TimeStamp;
 
@@ -96,10 +90,7 @@ public:
 	}
 
 	// debug log utility
-	const TCHAR* ToString() const
-	{
-		return EWorkspaceState::ToString(WorkspaceState);
-	}
+	const TCHAR* ToString() const;
 
 	/** ISourceControlState interface */
 	virtual int32 GetHistorySize() const override;
@@ -165,9 +156,9 @@ public:
 	FString LockedWhere;
 
 	/** State of the workspace */
-	EWorkspaceState::Type WorkspaceState = EWorkspaceState::Unknown;
+	EWorkspaceState WorkspaceState = EWorkspaceState::Unknown;
 
-	/** Latest revision number of the file in the depot */
+	/** Latest revision number of the file in the depot (on the current branch) */
 	int DepotRevisionChangeset = INVALID_REVISION;
 
 	/** Latest revision number at which a file was synced to before being edited */
diff --git a/Engine/Plugins/Developer/PlasticSourceControl/Source/PlasticSourceControl/Private/PlasticSourceControlUtils.cpp b/Engine/Plugins/Developer/PlasticSourceControl/Source/PlasticSourceControl/Private/PlasticSourceControlUtils.cpp
--- a/Engine/Plugins/Developer/PlasticSourceControl/Source/PlasticSourceControl/Private/PlasticSourceControlUtils.cpp
+++ b/Engine/Plugins/Developer/PlasticSourceControl/Source/PlasticSourceControl/Private/PlasticSourceControlUtils.cpp
@@ -1,363 +1,54 @@
 // Copyright Epic Games, Inc. All Rights Reserved.
 
 #include "PlasticSourceControlUtils.h"
+
 #include "PlasticSourceControlCommand.h"
 #include "PlasticSourceControlModule.h"
 #include "PlasticSourceControlProjectSettings.h"
 #include "PlasticSourceControlProvider.h"
 #include "PlasticSourceControlSettings.h"
+#include "PlasticSourceControlShell.h"
 #include "PlasticSourceControlState.h"
+#include "PlasticSourceControlVersions.h"
+#include "ISourceControlModule.h"
+
 #include "HAL/PlatformFileManager.h"
 #include "HAL/PlatformProcess.h"
-#include "HAL/FileManager.h"
-#include "Misc/ScopeLock.h"
 #include "Misc/FileHelper.h"
 #include "Misc/Paths.h"
-#include "HAL/PlatformTime.h"
-#include "Modules/ModuleManager.h"
 #include "XmlParser.h"
-#include "ISourceControlModule.h"
 #include "SoftwareVersion.h"
+#include "ScopedTempFile.h"
+
 #include "PlasticSourceControlChangelist.h"
 #include "PlasticSourceControlChangelistState.h"
 
-#if PLATFORM_LINUX
-#include <sys/ioctl.h>
-#endif
-
-#if PLATFORM_WINDOWS
-#include "Windows/WindowsHWrapper.h" // SECURITY_ATTRIBUTES
-#undef GetUserName
-#endif
-
-
-namespace PlasticSourceControlConstants
-{
-#if PLATFORM_WINDOWS
-	const TCHAR* pchDelim = TEXT("\r\n");
-#else
-	const TCHAR* pchDelim = TEXT("\n");
-#endif
-}
-
-FScopedTempFile::FScopedTempFile(const FString& InText)
-{
-	Filename = FPaths::CreateTempFilename(*FPaths::ProjectLogDir(), TEXT("Plastic-Temp"), TEXT(".txt"));
-	if (!FFileHelper::SaveStringToFile(InText, *Filename, FFileHelper::EEncodingOptions::ForceUTF8WithoutBOM))
-	{
-		UE_LOG(LogSourceControl, Error, TEXT("Failed to write to temp file: %s"), *Filename);
-	}
-}
-
-FScopedTempFile::FScopedTempFile(const FText& InText)
-{
-	Filename = FPaths::CreateTempFilename(*FPaths::ProjectLogDir(), TEXT("Plastic-Temp"), TEXT(".txt"));
-	if (!FFileHelper::SaveStringToFile(InText.ToString(), *Filename, FFileHelper::EEncodingOptions::ForceUTF8WithoutBOM))
-	{
-		UE_LOG(LogSourceControl, Error, TEXT("Failed to write to temp file: %s"), *Filename);
-	}
-}
-
-FScopedTempFile::~FScopedTempFile()
-{
-	if (FPaths::FileExists(Filename))
-	{
-		if (!FPlatformFileManager::Get().GetPlatformFile().DeleteFile(*Filename))
-		{
-			UE_LOG(LogSourceControl, Error, TEXT("Failed to delete temp file: %s"), *Filename);
-		}
-	}
-}
-
-const FString& FScopedTempFile::GetFilename() const
-{
-	return Filename;
-}
-
 namespace PlasticSourceControlUtils
 {
-// In/Out Pipes for the 'cm shell' persistent child process
-static void*			ShellOutputPipeRead = nullptr;
-static void*			ShellOutputPipeWrite = nullptr;
-static void*			ShellInputPipeRead = nullptr;
-static void*			ShellInputPipeWrite = nullptr;
-static FProcHandle		ShellProcessHandle;
-static FCriticalSection	ShellCriticalSection;
-static size_t			ShellCommandCounter = -1;
-static double			ShellCumulatedTime = 0.;
-
-// Internal function to cleanup (called under the critical section)
-static void _CleanupBackgroundCommandLineShell()
-{
-	FPlatformProcess::ClosePipe(ShellOutputPipeRead, ShellOutputPipeWrite);
-	FPlatformProcess::ClosePipe(ShellInputPipeRead, ShellInputPipeWrite);
-	ShellOutputPipeRead = ShellOutputPipeWrite = nullptr;
-	ShellInputPipeRead = ShellInputPipeWrite = nullptr;
-}
-
-// Internal function to launch the Plastic SCM background 'cm' process in interactive shell mode (called under the critical section)
-static bool _StartBackgroundPlasticShell(const FString& InPathToPlasticBinary, const FString& InWorkingDirectory)
-{
-	TRACE_CPUPROFILER_EVENT_SCOPE(FPlasticSourceControlUtils::_StartBackgroundPlasticShell);
-
-	const FString FullCommand(TEXT("shell --encoding=UTF-8"));
-
-	const bool bLaunchDetached = false;				// the new process will NOT have its own window
-	const bool bLaunchHidden = true;				// the new process will be minimized in the task bar
-	const bool bLaunchReallyHidden = bLaunchHidden; // the new process will not have a window or be in the task bar
-
-	const double StartTimestamp = FPlatformTime::Seconds();
-
-	verify(FPlatformProcess::CreatePipe(ShellOutputPipeRead, ShellOutputPipeWrite, false));	// For reading outputs from cm shell child process
-	verify(FPlatformProcess::CreatePipe(ShellInputPipeRead, ShellInputPipeWrite, true));	// For writing commands to cm shell child process
-
-	ShellProcessHandle = FPlatformProcess::CreateProc(*InPathToPlasticBinary, *FullCommand, bLaunchDetached, bLaunchHidden, bLaunchReallyHidden, nullptr, 0, *InWorkingDirectory, ShellOutputPipeWrite, ShellInputPipeRead);
-	if (!ShellProcessHandle.IsValid())
-	{
-		UE_LOG(LogSourceControl, Warning, TEXT("Failed to launch 'cm shell'")); // not a bug, just no Plastic SCM cli found
-		_CleanupBackgroundCommandLineShell();
-	}
-	else
-	{
-		const double ElapsedTime = (FPlatformTime::Seconds() - StartTimestamp);
-		UE_LOG(LogSourceControl, Verbose, TEXT("_StartBackgroundPlasticShell: '%s %s' ok (in %.3lfs, handle %d)"), *InPathToPlasticBinary, *FullCommand, ElapsedTime, ShellProcessHandle.Get());
-		ShellCommandCounter = 0;
-		ShellCumulatedTime = ElapsedTime;
-	}
-
-	return ShellProcessHandle.IsValid();
-}
-
-// Internal function (called under the critical section)
-static void _ExitBackgroundCommandLineShell()
-{
-	TRACE_CPUPROFILER_EVENT_SCOPE(FPlasticSourceControlUtils::_ExitBackgroundCommandLineShell);
-
-	if (ShellProcessHandle.IsValid())
-	{
-		if (FPlatformProcess::IsProcRunning(ShellProcessHandle))
-		{
-			// Tell the 'cm shell' to exit
-			FPlatformProcess::WritePipe(ShellInputPipeWrite, TEXT("exit"));
-			// And wait up to one second for its termination
-			const double Timeout = 1.0;
-			const double StartTimestamp = FPlatformTime::Seconds();
-			while (FPlatformProcess::IsProcRunning(ShellProcessHandle))
-			{
-				if ((FPlatformTime::Seconds() - StartTimestamp) > Timeout)
-				{
-					UE_LOG(LogSourceControl, Warning, TEXT("ExitBackgroundCommandLineShell: cm shell didn't stop gracefully in %lfs."), Timeout);
-					FPlatformProcess::TerminateProc(ShellProcessHandle);
-					break;
-				}
-				FPlatformProcess::Sleep(0.01f);
-			}
-		}
-		FPlatformProcess::CloseProc(ShellProcessHandle);
-		_CleanupBackgroundCommandLineShell();
-	}
-}
-
-// Internal function (called under the critical section)
-static void _RestartBackgroundCommandLineShell()
-{
-	const FPlasticSourceControlProvider& Provider = FPlasticSourceControlModule::Get().GetProvider();
-	const FString& PathToPlasticBinary = Provider.AccessSettings().GetBinaryPath();
-	const FString& WorkingDirectory = Provider.GetPathToWorkspaceRoot();
-
-	_ExitBackgroundCommandLineShell();
-	_StartBackgroundPlasticShell(PathToPlasticBinary, WorkingDirectory);
-}
-
-// Internal function (called under the critical section)
-static bool _RunCommandInternal(const FString& InCommand, const TArray<FString>& InParameters, const TArray<FString>& InFiles, const EConcurrency::Type InConcurrency, FString& OutResults, FString& OutErrors)
-{
-	TRACE_CPUPROFILER_EVENT_SCOPE(FPlasticSourceControlUtils::_RunCommandInternal);
-
-	bool bResult = false;
-
-	ShellCommandCounter++;
-
-	// Detect previous crash of cm.exe and restart 'cm shell'
-	if (!FPlatformProcess::IsProcRunning(ShellProcessHandle))
-	{
-		UE_LOG(LogSourceControl, Warning, TEXT("RunCommand: 'cm shell' has stopped. Restarting!"));
-		_RestartBackgroundCommandLineShell();
-	}
-
-	// Start with the Plastic command itself ("status", "log", "checkin"...)
-	FString FullCommand = InCommand;
-	// Append to the command all parameters, and then finally the files
-	for (const FString& Parameter : InParameters)
-	{
-		FullCommand += TEXT(" ");
-		FullCommand += Parameter;
-	}
-	for (const FString& File : InFiles)
-	{
-		FullCommand += TEXT(" \"");
-		FullCommand += File;
-		FullCommand += TEXT("\"");
-	}
-	const FString LoggableCommand = FullCommand.Left(256); // Limit command log size to 256 characters
-	UE_LOG(LogSourceControl, Verbose, TEXT("RunCommand: '%s' (%d chars, %d files)"), *LoggableCommand, FullCommand.Len()+1, InFiles.Num());
-	FullCommand += TEXT('\n'); // Finalize the command line
-
-	// Send command to 'cm shell' process in UTF-8
-	// NOTE: this explicit conversion to UTF-8 shouldn't be needed since FPlatformProcess::WritePipe() says it does it, but reading the implementation for Windows Platform show it merily truncates 16bits to 8bits chars!
-	// NOTE: on the other hand, ReadPipe() does the conversion from UTF-8 correctly already!
-	const FTCHARToUTF8 FullCommandUtf8(*FullCommand);
-	const bool bWriteOk = FPlatformProcess::WritePipe(ShellInputPipeWrite, reinterpret_cast<const uint8*>(FullCommandUtf8.Get()), FullCommandUtf8.Length());
-
-	// And wait up to 180.0 seconds for any kind of output from cm shell: in case of lengthier operation, intermediate output (like percentage of progress) is expected, which would refresh the timeout
-	static const double Timeout = 180.0;
-	const double StartTimestamp = FPlatformTime::Seconds();
-	double LastActivity = StartTimestamp;
-	double LastLog = StartTimestamp;
-	static const double LogInterval = 5.0;
-	int32 PreviousLogLen = 0;
-	while (FPlatformProcess::IsProcRunning(ShellProcessHandle))
-	{
-		FString Output = FPlatformProcess::ReadPipe(ShellOutputPipeRead);
-		if (!Output.IsEmpty())
-		{
-			TRACE_CPUPROFILER_EVENT_SCOPE(FPlasticSourceControlUtils::_RunCommandInternal::ParseOutput);
-
-			LastActivity = FPlatformTime::Seconds(); // freshen the timestamp while cm is still actively outputting information
-			OutResults.Append(MoveTemp(Output));
-			// Search the output for the line containing the result code, also indicating the end of the command
-			const uint32 IndexCommandResult = OutResults.Find(TEXT("CommandResult "), ESearchCase::CaseSensitive, ESearchDir::FromEnd);
-			if (INDEX_NONE != IndexCommandResult)
-			{
-				const uint32 IndexEndResult = OutResults.Find(PlasticSourceControlConstants::pchDelim, ESearchCase::CaseSensitive, ESearchDir::FromStart, IndexCommandResult + 14);
-				if (INDEX_NONE != IndexEndResult)
-				{
-					const FString Result = OutResults.Mid(IndexCommandResult + 14, IndexEndResult - IndexCommandResult - 14);
-					const int32 ResultCode = FCString::Atoi(*Result);
-					bResult = (ResultCode == 0);
-					// remove the CommandResult line from the OutResults
-					OutResults.RemoveAt(IndexCommandResult, OutResults.Len() - IndexCommandResult);
-					break;
-				}
-			}
-		}
-		else if ((FPlatformTime::Seconds() - LastLog > LogInterval) && (PreviousLogLen < OutResults.Len()))
-		{
-			// In case of long running operation, start to print intermediate output from cm shell (like percentage of progress)
-			UE_LOG(LogSourceControl, Log, TEXT("RunCommand: '%s' in progress for %.3lfs...\n%s"), *InCommand, (FPlatformTime::Seconds() - StartTimestamp), *OutResults.Mid(PreviousLogLen));
-			PreviousLogLen = OutResults.Len();
-			LastLog = FPlatformTime::Seconds(); // freshen the timestamp of last log
-		}
-		else if (FPlatformTime::Seconds() - LastActivity > Timeout)
-		{
-			// In case of timeout, ask the blocking 'cm shell' process to exit, detach from it and restart it immediately
-			UE_LOG(LogSourceControl, Error, TEXT("RunCommand: '%s' TIMEOUT after %.3lfs output (%d chars):\n%s"), *InCommand, (FPlatformTime::Seconds() - StartTimestamp), OutResults.Len(), *OutResults.Mid(PreviousLogLen));
-			_RestartBackgroundCommandLineShell();
-			return false;
-		}
-		else if (IsEngineExitRequested())
-		{
-			UE_LOG(LogSourceControl, Warning, TEXT("RunCommand: '%s' Engine Exit was requested after %.3lfs output (%d chars):\n%s"), *InCommand, (FPlatformTime::Seconds() - StartTimestamp), OutResults.Len(), *OutResults.Mid(PreviousLogLen));
-			_ExitBackgroundCommandLineShell();
-		}
-
-		TRACE_CPUPROFILER_EVENT_SCOPE(FPlasticSourceControlUtils::_RunCommandInternal::Sleep);
-		FPlatformProcess::Sleep(0.001f);
-	}
-	const double ElapsedTime = (FPlatformTime::Seconds() - StartTimestamp);
-
-	if (!InCommand.Equals(TEXT("exit")))
-	{
-		if (!FPlatformProcess::IsProcRunning(ShellProcessHandle))
-		{
-			// 'cm shell' normally only terminates in case of 'exit' command. Will restart on next command.
-			UE_LOG(LogSourceControl, Error, TEXT("RunCommand: '%s' 'cm shell' stopped after %.3lfs output (%d chars):\n%s"), *LoggableCommand, ElapsedTime, OutResults.Len(), *OutResults.Left(4096)); // Limit result size to 4096 characters
-		}
-		else if (!bResult)
-		{
-			UE_LOG(LogSourceControl, Warning, TEXT("RunCommand: '%s' (in %.3lfs) output (%d chars):\n%s"), *LoggableCommand, ElapsedTime, OutResults.Len(), *OutResults.Left(4096)); // Limit result size to 4096 characters
-		}
-		else
-		{
-			if (PreviousLogLen > 0)
-			{
-				UE_LOG(LogSourceControl, Log, TEXT("RunCommand: '%s' (in %.3lfs) output (%d chars):\n%s"), *LoggableCommand, ElapsedTime, OutResults.Len(), *OutResults.Mid(PreviousLogLen).Left(4096)); // Limit result size to 4096 characters
-			}
-			else
-			{
-				if (OutResults.Len() <= 200) // Limit result size to 200 characters
-				{
-					UE_LOG(LogSourceControl, Log, TEXT("RunCommand: '%s' (in %.3lfs) output (%d chars):\n%s"), *LoggableCommand, ElapsedTime, OutResults.Len(), *OutResults);
-				}
-				else
-				{
-					UE_LOG(LogSourceControl, Log, TEXT("RunCommand: '%s' (in %.3lfs) (output %d chars not displayed)"), *LoggableCommand, ElapsedTime, OutResults.Len());
-					UE_LOG(LogSourceControl, Verbose, TEXT("\n%s"), *OutResults.Left(4096));; // Limit result size to 4096 characters
-				}
-			}
-		}
-	}
-	// Return output as error if result code is an error
-	if (!bResult)
-	{
-		OutErrors = MoveTemp(OutResults);
-	}
-
-	ShellCumulatedTime += ElapsedTime;
-	UE_LOG(LogSourceControl, Verbose, TEXT("RunCommand: cumulated time spent in shell: %.3lfs (count %d)"), ShellCumulatedTime, ShellCommandCounter);
-
-	return bResult;
-}
-
-// Launch the Plastic SCM background 'cm shell' process in background for optimized successive commands (thread-safe)
-bool LaunchBackgroundPlasticShell(const FString& InPathToPlasticBinary, const FString& InWorkingDirectory)
-{
-	// Protect public APIs from multi-thread access
-	FScopeLock Lock(&ShellCriticalSection);
-
-	// terminate previous shell if one is already running
-	_ExitBackgroundCommandLineShell();
-
-	return _StartBackgroundPlasticShell(InPathToPlasticBinary, InWorkingDirectory);
-}
-
-// Terminate the background 'cm shell' process and associated pipes (thread-safe)
-void Terminate()
-{
-	// Protect public APIs from multi-thread access
-	FScopeLock Lock(&ShellCriticalSection);
-
-	_ExitBackgroundCommandLineShell();
-}
 
-// Run command (thread-safe)
-bool RunCommandInternal(const FString& InCommand, const TArray<FString>& InParameters, const TArray<FString>& InFiles, const EConcurrency::Type InConcurrency, FString& OutResults, FString& OutErrors)
+// Run a command and return the result as raw strings
+bool RunCommand(const FString& InCommand, const TArray<FString>& InParameters, const TArray<FString>& InFiles, FString& OutResults, FString& OutErrors)
 {
-	// Protect public APIs from multi-thread access
-	FScopeLock Lock(&ShellCriticalSection);
-
-	return _RunCommandInternal(InCommand, InParameters, InFiles, InConcurrency, OutResults, OutErrors);
+	return PlasticSourceControlShell::RunCommand(InCommand, InParameters, InFiles, OutResults, OutErrors);
 }
 
-// Basic parsing or results & errors from the Plastic command line process
-bool RunCommand(const FString& InCommand, const TArray<FString>& InParameters, const TArray<FString>& InFiles, const EConcurrency::Type InConcurrency, TArray<FString>& OutResults, TArray<FString>& OutErrorMessages)
+// Run a command with basic parsing or results & errors from the Plastic command line process
+bool RunCommand(const FString& InCommand, const TArray<FString>& InParameters, const TArray<FString>& InFiles, TArray<FString>& OutResults, TArray<FString>& OutErrorMessages)
 {
 	FString Results;
 	FString Errors;
 
-	const bool bResult = RunCommandInternal(InCommand, InParameters, InFiles, InConcurrency, Results, Errors);
+	const bool bResult = PlasticSourceControlShell::RunCommand(InCommand, InParameters, InFiles, Results, Errors);
 
-	Results.ParseIntoArray(OutResults, PlasticSourceControlConstants::pchDelim, true);
-	Errors.ParseIntoArray(OutErrorMessages, PlasticSourceControlConstants::pchDelim, true);
+	Results.ParseIntoArray(OutResults, PlasticSourceControlShell::pchDelim, true);
+	Errors.ParseIntoArray(OutErrorMessages, PlasticSourceControlShell::pchDelim, true);
 
 	return bResult;
 }
 
 const FSoftwareVersion& GetOldestSupportedPlasticScmVersion()
 {
-	// https://www.plasticscm.com/download/releasenotes/9.0.16.4839 cm changelist 'persistent' flag now contain a '--' prefix.
-	static FSoftwareVersion s_OldestSupportedPlasticScmVersion(TEXT("9.0.16.4839"));
-	return s_OldestSupportedPlasticScmVersion;
+	return PlasticSourceControlVersions::OldestSupported;
 }
 
 FString FindPlasticBinaryPath()
@@ -370,44 +61,21 @@ FString FindPlasticBinaryPath()
 }
 
 // Find the root of the Plastic workspace, looking from the provided path and upward in its parent directories.
-bool FindRootDirectory(const FString& InPath, FString& OutWorkspaceRoot)
+bool GetWorkspacePath(const FString& InPath, FString& OutWorkspaceRoot)
 {
-	bool bFound = false;
-	FString PathToPlasticSubdirectory;
-	OutWorkspaceRoot = InPath;
-
-	auto TrimTrailing = [](FString& Str, const TCHAR Char)
-	{
-		int32 Len = Str.Len();
-		while (Len && Str[Len - 1] == Char)
-		{
-			Str = Str.LeftChop(1);
-			Len = Str.Len();
-		}
-	};
-
-	TrimTrailing(OutWorkspaceRoot, '\\');
-	TrimTrailing(OutWorkspaceRoot, '/');
-
-	while (!bFound && !OutWorkspaceRoot.IsEmpty())
+	TArray<FString> Results;
+	TArray<FString> ErrorMessages;
+	TArray<FString> Parameters;
+	Parameters.Add(TEXT("--format={wkpath}"));
+	Parameters.Add(TEXT("."));
+	const bool bFound = RunCommand(TEXT("getworkspacefrompath"), Parameters, TArray<FString>(), Results, ErrorMessages);
+	if (bFound && Results.Num() > 0)
 	{
-		// Look for the ".plastic" subdirectory present at the root of every Plastic workspace
-		PathToPlasticSubdirectory = OutWorkspaceRoot / TEXT(".plastic");
-		bFound = FPaths::DirectoryExists(*PathToPlasticSubdirectory);
-		if (!bFound)
-		{
-			int32 LastSlashIndex;
-			if (OutWorkspaceRoot.FindLastChar(TEXT('/'), LastSlashIndex))
-			{
-				OutWorkspaceRoot = OutWorkspaceRoot.Left(LastSlashIndex);
-			}
-			else
-			{
-				OutWorkspaceRoot.Empty();
-			}
-		}
+		OutWorkspaceRoot = MoveTemp(Results[0]);
+		FPaths::NormalizeDirectoryName(OutWorkspaceRoot);
+		OutWorkspaceRoot.AppendChar(TEXT('/'));
 	}
-	if (!bFound)
+	else
 	{
 		OutWorkspaceRoot = InPath; // If not found, return the provided dir as best possible root.
 	}
@@ -417,62 +85,91 @@ bool FindRootDirectory(const FString& InPath, FString& OutWorkspaceRoot)
 // This is called once by FPlasticSourceControlProvider::CheckPlasticAvailability()
 bool GetPlasticScmVersion(FSoftwareVersion& OutPlasticScmVersion)
 {
-	TArray<FString> InfoMessages;
+	TArray<FString> Results;
+	TArray<FString> ErrorMessages;
+	const bool bResult = RunCommand(TEXT("version"), TArray<FString>(), TArray<FString>(), Results, ErrorMessages);
+	if (bResult && Results.Num() > 0)
+	{
+		OutPlasticScmVersion = FSoftwareVersion(MoveTemp(Results[0]));
+		return true;
+	}
+	return false;
+}
+
+bool GetCmLocation(FString& OutCmLocation)
+{
+	TArray<FString> Results;
 	TArray<FString> ErrorMessages;
-	const bool bResult = RunCommand(TEXT("version"), TArray<FString>(), TArray<FString>(), EConcurrency::Synchronous, InfoMessages, ErrorMessages);
-	if (bResult && InfoMessages.Num() > 0)
+	const bool bResult = RunCommand(TEXT("location"), TArray<FString>(), TArray<FString>(), Results, ErrorMessages);
+	if (bResult && Results.Num() > 0)
 	{
-		OutPlasticScmVersion = FSoftwareVersion(InfoMessages[0]);
+		OutCmLocation = MoveTemp(Results[0]);
 		return true;
 	}
 	return false;
 }
 
+bool GetConfigSetFilesAsReadOnly()
+{
+	TArray<FString> Results;
+	TArray<FString> ErrorMessages;
+	TArray<FString> Parameters;
+	Parameters.Add(TEXT("setfileasreadonly"));
+	const bool bResult = RunCommand(TEXT("getconfig"), Parameters, TArray<FString>(), Results, ErrorMessages);
+	if (bResult && Results.Num() > 0)
+	{
+		if ((Results[0].Compare(TEXT("yes"), ESearchCase::IgnoreCase) == 0) || (Results[0].Compare(TEXT("true"), ESearchCase::IgnoreCase) == 0))
+		{
+			return true;
+		}
+	}
+	return false;
+}
+
 void GetUserName(FString& OutUserName)
 {
-	TArray<FString> InfoMessages;
+	TArray<FString> Results;
 	TArray<FString> ErrorMessages;
-	const bool bResult = RunCommand(TEXT("whoami"), TArray<FString>(), TArray<FString>(), EConcurrency::Synchronous, InfoMessages, ErrorMessages);
-	if (bResult && InfoMessages.Num() > 0)
+	const bool bResult = RunCommand(TEXT("whoami"), TArray<FString>(), TArray<FString>(), Results, ErrorMessages);
+	if (bResult && Results.Num() > 0)
 	{
-		OutUserName = InfoMessages[0];
+		OutUserName = MoveTemp(Results[0]);
 	}
 }
 
 bool GetWorkspaceName(const FString& InWorkspaceRoot, FString& OutWorkspaceName, TArray<FString>& OutErrorMessages)
 {
-	TArray<FString> InfoMessages;
-
+	TArray<FString> Results;
 	TArray<FString> Parameters;
-	Parameters.Add(TEXT("--format={0}"));
+	Parameters.Add(TEXT("--format={wkname}"));
 	TArray<FString> Files;
-	Files.Add(InWorkspaceRoot);
+	Files.Add(InWorkspaceRoot); // Uses an absolute path so that the error message is explicit
 	// Get the workspace name
-	const bool bResult = RunCommand(TEXT("getworkspacefrompath"), Parameters, Files, EConcurrency::Synchronous, InfoMessages, OutErrorMessages);
-	if (bResult && InfoMessages.Num() > 0)
+	const bool bResult = RunCommand(TEXT("getworkspacefrompath"), Parameters, Files, Results, OutErrorMessages);
+	if (bResult && Results.Num() > 0)
 	{
 		// NOTE: an old version of cm getworkspacefrompath didn't return an error code so we had to rely on the error message
-		if (!InfoMessages[0].EndsWith(TEXT(" is not in a workspace.")))
+		if (!Results[0].EndsWith(TEXT(" is not in a workspace.")))
 		{
-			OutWorkspaceName = MoveTemp(InfoMessages[0]);
+			OutWorkspaceName = MoveTemp(Results[0]);
 		}
 	}
 
 	return bResult;
 }
 
-static bool ParseWorkspaceInformation(const TArray<FString>& InInfoMessages, int32& OutChangeset, FString& OutRepositoryName, FString& OutServerUrl, FString& OutBranchName)
+static bool ParseWorkspaceInformation(const TArray<FString>& InResults, int32& OutChangeset, FString& OutRepositoryName, FString& OutServerUrl, FString& OutBranchName)
 {
 	bool bResult = true;
 
 	// Get workspace status, in the form "cs:41@rep:UEPlasticPlugin@repserver:localhost:8087" (disabled by the "--nostatus" flag)
 	//                                or "cs:41@rep:UEPlasticPlugin@repserver:SRombauts@cloud" (when connected directly to the cloud)
-	if (InInfoMessages.Num() > 0)
+	if (InResults.Num() > 0)
 	{
 		static const FString ChangesetPrefix(TEXT("cs:"));
 		static const FString RepPrefix(TEXT("@rep:"));
 		static const FString ServerPrefix(TEXT("@repserver:"));
-		const FString& WorkspaceStatus = InInfoMessages[0];
+		const FString& WorkspaceStatus = InResults[0];
 		const int32 RepIndex = WorkspaceStatus.Find(RepPrefix, ESearchCase::CaseSensitive);
 		const int32 ServerIndex = WorkspaceStatus.Find(ServerPrefix, ESearchCase::CaseSensitive, ESearchDir::FromEnd);
 		if ((RepIndex > INDEX_NONE) && (ServerIndex > INDEX_NONE))
@@ -487,11 +184,11 @@ static bool ParseWorkspaceInformation(const TArray<FString>& InInfoMessages, int
 			bResult = false;
 		}
 	}
-	// Get the branch name, in the form "Branch /main@UEPlasticPluginDev" (enabled by the "--wkconfig" flag)
-	if (InInfoMessages.Num() > 1)
+	// Get the branch name, in the form "Branch /main@UE5PlasticPluginDev@test@cloud" (enabled by the "--wkconfig" flag)
+	if (InResults.Num() > 1)
 	{
 		static const FString BranchPrefix(TEXT("Branch "));
-		const FString& BranchInfo = InInfoMessages[1];
+		const FString& BranchInfo = InResults[1];
 		const int32 BranchIndex = BranchInfo.Find(BranchPrefix, ESearchCase::CaseSensitive);
 		if (BranchIndex > INDEX_NONE)
 		{
@@ -504,16 +201,71 @@ static bool ParseWorkspaceInformation(const TArray<FString>& InInfoMessages, int
 
 bool GetWorkspaceInformation(int32& OutChangeset, FString& OutRepositoryName, FString& OutServerUrl, FString& OutBranchName, TArray<FString>& OutErrorMessages)
 {
-	TArray<FString> InfoMessages;
+	TArray<FString> Results;
 	TArray<FString> Parameters;
 	Parameters.Add(TEXT("--compact"));
 	Parameters.Add(TEXT("--header")); // Only prints the workspace status. No file status.
 	// NOTE: --wkconfig results in two network calls GetBranchInfoByName & GetLastChangesetOnBranch so it's okay to do it once here but not all the time
 	Parameters.Add(TEXT("--wkconfig")); // Branch name. NOTE: Deprecated in 8.0.16.3000 https://www.plasticscm.com/download/releasenotes/8.0.16.3000
-	bool bResult = RunCommand(TEXT("status"), Parameters, TArray<FString>(), EConcurrency::Synchronous, InfoMessages, OutErrorMessages);
+	bool bResult = RunCommand(TEXT("status"), Parameters, TArray<FString>(), Results, OutErrorMessages);
 	if (bResult)
 	{
-		bResult = ParseWorkspaceInformation(InfoMessages, OutChangeset, OutRepositoryName, OutServerUrl, OutBranchName);
+		bResult = ParseWorkspaceInformation(Results, OutChangeset, OutRepositoryName, OutServerUrl, OutBranchName);
+	}
+
+	return bResult;
+}
+
+static bool ParseWorkspaceInfo(const TArray<FString>& InResults, FString& OutBranchName, FString& OutRepositoryName, FString& OutServerUrl)
+{
+	if (InResults.Num() == 0)
+	{
+		return false;
+	}
+
+	// Get workspace information, in the form "Branch /main@UE5PlasticPluginDev@localhost:8087"
+	//                                     or "Branch /main@UE5PlasticPluginDev@test@cloud" (when connected directly to the cloud)
+	//                                     or "Branch /main@rep:UE5OpenWorldPerfTest@repserver:test@cloud"
+	static const FString BranchPrefix(TEXT("Branch "));
+	if (!InResults[0].StartsWith(BranchPrefix, ESearchCase::CaseSensitive))
+	{
+		return false;
+	}
+	FString Temp = InResults[0].RightChop(BranchPrefix.Len());
+	int32 SeparatorIndex;
+	if (!Temp.FindChar(TEXT('@'), SeparatorIndex))
+	{
+		return false;
+	}
+	OutBranchName = Temp.Left(SeparatorIndex);
+	Temp.RightChopInline(SeparatorIndex + 1);
+	if (!Temp.FindChar(TEXT('@'), SeparatorIndex))
+	{
+		return false;
+	}
+	OutRepositoryName = Temp.Left(SeparatorIndex);
+	static const FString RepPrefix(TEXT("rep:"));
+	if (OutRepositoryName.StartsWith(RepPrefix, ESearchCase::CaseSensitive))
+	{
+		OutRepositoryName.RightChopInline(RepPrefix.Len());
+	}
+	OutServerUrl = Temp.RightChop(SeparatorIndex + 1);
+	static const FString RepserverPrefix(TEXT("repserver:"));
+	if (OutServerUrl.StartsWith(RepserverPrefix, ESearchCase::CaseSensitive))
+	{
+		OutServerUrl.RightChopInline(RepserverPrefix.Len());
+	}
+
+	return true;
+}
+
+bool GetWorkspaceInfo(FString& OutRepositoryName, FString& OutServerUrl, FString& OutBranchName, TArray<FString>& OutErrorMessages)
+{
+	TArray<FString> Results;
+	bool bResult = RunCommand(TEXT("workspaceinfo"), TArray<FString>(), TArray<FString>(), Results, OutErrorMessages);
+	if (bResult)
+	{
+		bResult = ParseWorkspaceInfo(Results, OutRepositoryName, OutServerUrl, OutBranchName);
 	}
 
 	return bResult;
@@ -552,7 +304,7 @@ static FString RenamedFromPlasticStatus(const FString& InResult)
 {
 	FString RenamedFrom;
 	int32 RenameIndex;
-	if (InResult.FindLastChar('>', RenameIndex))
+	if (InResult.FindLastChar(TEXT('>'), RenameIndex))
 	{
 		// Extract only the first part of a rename "from -> to" (after the 2 letters status surrounded by 2 spaces)
 		RenamedFrom = InResult.Mid(9, RenameIndex - 9 - 2);
@@ -565,6 +317,7 @@ static FString RenamedFromPlasticStatus(const FString& InResult)
  *
  * Examples of status results:
  CO Content\CheckedOut_BP.uasset
+ CO Content\Merged_BP.uasset (Merge from 140)
  MV 100% Content\ToMove_BP.uasset -> Content\Moved_BP.uasset
  *
  * @param[in] InResult One line of status
@@ -576,7 +329,7 @@ static FString FilenameFromPlasticStatus(const FString& InResult)
 {
 	FString RelativeFilename;
 	int32 RenameIndex;
-	if (InResult.FindLastChar('>', RenameIndex))
+	if (InResult.FindLastChar(TEXT('>'), RenameIndex))
 	{
 		// Extract only the second part of a rename "from -> to"
 		RelativeFilename = InResult.RightChop(RenameIndex + 2);
@@ -585,6 +338,12 @@ static FString FilenameFromPlasticStatus(const FString& InResult)
 	{
 		// Extract the relative filename from the Plastic SCM status result (after the 2 letters status surrounded by 2 spaces)
 		RelativeFilename = InResult.RightChop(4);
+
+		const int32 MergeIndex = RelativeFilename.Find(TEXT(" (Merge from "));
+		if (MergeIndex != INDEX_NONE)
+		{
+			RelativeFilename.LeftInline(MergeIndex);
+		}
 	}
 
 	return RelativeFilename;
@@ -629,9 +388,9 @@ private:
  MV 100% Content\ToMove_BP.uasset -> Content\Moved_BP.uasset
  LM 100% Content\ToMove2_BP.uasset -> Content\Moved2_BP.uasset
  */
-static EWorkspaceState::Type StateFromPlasticStatus(const FString& InResult)
+static EWorkspaceState StateFromPlasticStatus(const FString& InResult)
 {
-	EWorkspaceState::Type State;
+	EWorkspaceState State;
 	const FString FileStatus = (InResult[0] == TEXT(' ')) ? InResult.Mid(1, 2) : InResult;
 
 	if (FileStatus == "CH") // Modified but not Checked-Out
@@ -707,7 +466,7 @@ static EWorkspaceState::Type StateFromPlasticStatus(const FString& InResult)
  */
 static void ParseFileStatusResult(TArray<FString>&& InFiles, const TArray<FString>& InResults, TArray<FPlasticSourceControlState>& OutStates, int32& OutChangeset, FString& OutBranchName)
 {
-	TRACE_CPUPROFILER_EVENT_SCOPE(FPlasticSourceControlUtils::ParseFileStatusResult);
+	TRACE_CPUPROFILER_EVENT_SCOPE(PlasticSourceControlUtils::ParseFileStatusResult);
 
 	const FString& WorkingDirectory = FPlasticSourceControlModule::Get().GetProvider().GetPathToWorkspaceRoot();
 
@@ -779,14 +538,14 @@ static void ParseFileStatusResult(TArray<FString>&& InFiles, const TArray<FStrin
 */
 static void ParseDirectoryStatusResultForDeleted(const TArray<FString>& InResults, TArray<FPlasticSourceControlState>& OutStates)
 {
-	TRACE_CPUPROFILER_EVENT_SCOPE(FPlasticSourceControlUtils::ParseDirectoryStatusResultForDeleted);
+	TRACE_CPUPROFILER_EVENT_SCOPE(PlasticSourceControlUtils::ParseDirectoryStatusResultForDeleted);
 
 	const FString& WorkingDirectory = FPlasticSourceControlModule::Get().GetProvider().GetPathToWorkspaceRoot();
 
 	// Iterate on each line of result of the status command
 	for (const FString& Result : InResults)
 	{
-		const EWorkspaceState::Type WorkspaceState = StateFromPlasticStatus(Result);
+		const EWorkspaceState WorkspaceState = StateFromPlasticStatus(Result);
 		if ((EWorkspaceState::Deleted == WorkspaceState) || (EWorkspaceState::LocallyDeleted == WorkspaceState))
 		{
 			FString RelativeFilename = FilenameFromPlasticStatus(Result);
@@ -833,9 +592,9 @@ public:
  * @param[out]	OutChangeset		The current Changeset Number
  * @param[out]	OutBranchName		Name of the current checked-out branch
  */
-static bool RunStatus(const FString& InDir, TArray<FString>&& InFiles, const EConcurrency::Type InConcurrency, TArray<FString>& OutErrorMessages, TArray<FPlasticSourceControlState>& OutStates, int32& OutChangeset, FString& OutBranchName)
+static bool RunStatus(const FString& InDir, TArray<FString>&& InFiles, TArray<FString>& OutErrorMessages, TArray<FPlasticSourceControlState>& OutStates, int32& OutChangeset, FString& OutBranchName)
 {
-	TRACE_CPUPROFILER_EVENT_SCOPE(FPlasticSourceControlUtils::RunStatus);
+	TRACE_CPUPROFILER_EVENT_SCOPE(PlasticSourceControlUtils::RunStatus);
 
 	check(InFiles.Num() > 0);
 
@@ -859,7 +618,7 @@ static bool RunStatus(const FString& InDir, TArray<FString>&& InFiles, const ECo
 	}
 	TArray<FString> Results;
 	TArray<FString> ErrorMessages;
-	const bool bResult = RunCommand(TEXT("status"), Parameters, OnePath, InConcurrency, Results, ErrorMessages);
+	const bool bResult = RunCommand(TEXT("status"), Parameters, OnePath, Results, ErrorMessages);
 	OutErrorMessages.Append(MoveTemp(ErrorMessages));
 	if (bResult)
 	{
@@ -933,7 +692,7 @@ public:
  */
 static void ParseFileinfoResults(const TArray<FString>& InResults, TArray<FPlasticSourceControlState>& InOutStates)
 {
-	TRACE_CPUPROFILER_EVENT_SCOPE(FPlasticSourceControlUtils::ParseFileinfoResults);
+	TRACE_CPUPROFILER_EVENT_SCOPE(PlasticSourceControlUtils::ParseFileinfoResults);
 
 	const FPlasticSourceControlProvider& Provider = FPlasticSourceControlModule::Get().GetProvider();
 
@@ -953,13 +712,6 @@ static void ParseFileinfoResults(const TArray<FString>& InResults, TArray<FPlast
 		FileState.LockedBy = MoveTemp(FileinfoParser.LockedBy);
 		FileState.LockedWhere = MoveTemp(FileinfoParser.LockedWhere);
 
-		// If a file is locked but not checked-out locally (or moved/renamed) this means it is locked by someone else or elsewhere
-		if ((FileState.WorkspaceState != EWorkspaceState::CheckedOut) && (FileState.WorkspaceState != EWorkspaceState::Moved) && !FileState.LockedBy.IsEmpty())
-		{
-			UE_LOG(LogSourceControl, Verbose, TEXT("LockedByOther(%s) by '%s!=%s' (or %s!=%s)"), *File, *FileState.LockedBy, *Provider.GetUserName(), *FileState.LockedWhere, *Provider.GetWorkspaceName());
-			FileState.WorkspaceState = EWorkspaceState::LockedByOther;
-		}
-
 		// debug log (only for the first few files)
 		if (IdxResult < 20)
 		{
@@ -980,13 +732,12 @@ static void ParseFileinfoResults(const TArray<FString>& InResults, TArray<FPlast
  *
  * @param[in]		bInWholeDirectory	If executed on a whole directory (typically Content/) for a "Submit Content" operation, optimize fileinfo more aggressively
  * @param			bInUpdateHistory	If getting the history of files, force execute the fileinfo command required to get RepSpec of XLinks (history view or visual diff)
- * @param[in]		InConcurrency		Is the command running in the background, or blocking the main thread
  * @param[out]		OutErrorMessages	Error messages from the "fileinfo" command
  * @param[in,out]	InOutStates			List of file states in the directory, gathered by the "status" command, completed by results of the "fileinfo" command
  */
-static bool RunFileinfo(const bool bInWholeDirectory, const bool bInUpdateHistory, const EConcurrency::Type InConcurrency, TArray<FString>& OutErrorMessages, TArray<FPlasticSourceControlState>& InOutStates)
+static bool RunFileinfo(const bool bInWholeDirectory, const bool bInUpdateHistory, TArray<FString>& OutErrorMessages, TArray<FPlasticSourceControlState>& InOutStates)
 {
-	TRACE_CPUPROFILER_EVENT_SCOPE(FPlasticSourceControlUtils::RunFileinfo);
+	TRACE_CPUPROFILER_EVENT_SCOPE(PlasticSourceControlUtils::RunFileinfo);
 
 	bool bResult = true;
 	TArray<FString> SelectedFiles;
@@ -1029,7 +780,7 @@ static bool RunFileinfo(const bool bInWholeDirectory, const bool bInUpdateHistor
 		TArray<FString> ErrorMessages;
 		TArray<FString> Parameters;
 		Parameters.Add(TEXT("--format=\"{RevisionChangeset};{RevisionHeadChangeset};{RepSpec};{LockedBy};{LockedWhere}\""));
-		bResult = RunCommand(TEXT("fileinfo"), Parameters, SelectedFiles, InConcurrency, Results, ErrorMessages);
+		bResult = RunCommand(TEXT("fileinfo"), Parameters, SelectedFiles, Results, ErrorMessages);
 		OutErrorMessages.Append(MoveTemp(ErrorMessages));
 		if (bResult)
 		{
@@ -1057,13 +808,13 @@ public:
 			{
 				Filename = Temp.Left(WhitespaceIndex);
 			}
-			Temp = Temp.RightChop(WhitespaceIndex + 1);
+			Temp.RightChopInline(WhitespaceIndex + 1);
 			if (Temp.FindChar(TEXT(' '), WhitespaceIndex))
 			{
 				const FString Base = Temp.Left(WhitespaceIndex);
 				BaseChangeset = FCString::Atoi(*Base);
 			}
-			Temp = Temp.RightChop(WhitespaceIndex + 1);
+			Temp.RightChopInline(WhitespaceIndex + 1);
 			if (Temp.FindChar(TEXT(' '), WhitespaceIndex))
 			{
 				const FString Source = Temp.Left(WhitespaceIndex);
@@ -1080,7 +831,7 @@ public:
 // Check if merging, and from which changelist, then execute a cm merge command to amend status for listed files
 static bool RunCheckMergeStatus(const TArray<FString>& InFiles, TArray<FString>& OutErrorMessages, TArray<FPlasticSourceControlState>& OutStates)
 {
-	TRACE_CPUPROFILER_EVENT_SCOPE(FPlasticSourceControlUtils::RunCheckMergeStatus);
+	TRACE_CPUPROFILER_EVENT_SCOPE(PlasticSourceControlUtils::RunCheckMergeStatus);
 
 	bool bResult = false;
 	const FPlasticSourceControlProvider& Provider = FPlasticSourceControlModule::Get().GetProvider();
@@ -1137,7 +888,7 @@ static bool RunCheckMergeStatus(const TArray<FString>& InFiles, TArray<FString>&
 					const TArray<FString> PendingMergeParameters = Parameters;
 					Parameters.Add(TEXT("--machinereadable"));
 					// call 'cm merge cs:xxx --machinereadable' (only dry-run, without the --merge parameter)
-					bResult = RunCommand(TEXT("merge"), Parameters, TArray<FString>(), EConcurrency::Synchronous, Results, ErrorMessages);
+					bResult = RunCommand(TEXT("merge"), Parameters, TArray<FString>(), Results, ErrorMessages);
 					OutErrorMessages.Append(MoveTemp(ErrorMessages));
 					// Parse the result, one line for each conflicted files:
 					for (const FString& Result : Results)
@@ -1194,7 +945,7 @@ struct FFilesInCommonDir
 };
 
 // Run a batch of Plastic "status" and "fileinfo" commands to update status of given files and directories.
-bool RunUpdateStatus(const TArray<FString>& InFiles, const bool bInUpdateHistory, const EConcurrency::Type InConcurrency, TArray<FString>& OutErrorMessages, TArray<FPlasticSourceControlState>& OutStates, int32& OutChangeset, FString& OutBranchName)
+bool RunUpdateStatus(const TArray<FString>& InFiles, const bool bInUpdateHistory, TArray<FString>& OutErrorMessages, TArray<FPlasticSourceControlState>& OutStates, int32& OutChangeset, FString& OutBranchName)
 {
 	bool bResults = true;
 
@@ -1245,7 +996,7 @@ bool RunUpdateStatus(const TArray<FString>& InFiles, const bool bInUpdateHistory
 				}
 				else
 				{
-					FString Path = FPaths::GetPath(File) + TEXT("/");
+					FString Path = FPaths::GetPath(File) + TEXT('/');
 					GroupOfFiles.Add(RootDir, { MoveTemp(Path), {File}});
 				}
 
@@ -1259,7 +1010,7 @@ bool RunUpdateStatus(const TArray<FString>& InFiles, const bool bInUpdateHistory
 		// This should be an edge case (typically the uproject file) .
 		if (!bDirFound)
 		{
-			FString Path = FPaths::GetPath(File) + TEXT("/");
+			FString Path = FPaths::GetPath(File) + TEXT('/');
 			FFilesInCommonDir* ExistingGroup = GroupOfFiles.Find(Path);
 			if (ExistingGroup != nullptr)
 			{
@@ -1289,14 +1040,18 @@ bool RunUpdateStatus(const TArray<FString>& InFiles, const bool bInUpdateHistory
 		// Run a "status" command on the directory to get workspace file states.
 		// (ie. Changed, CheckedOut, Copied, Replaced, Added, Private, Ignored, Deleted, LocallyDeleted, Moved, LocallyMoved)
 		TArray<FPlasticSourceControlState> States;
-		const bool bGroupOk = RunStatus(Group.Value.CommonDir, MoveTemp(Group.Value.Files), InConcurrency, OutErrorMessages, States, OutChangeset, OutBranchName);
-		if (bGroupOk && (States.Num() > 0))
+		const bool bGroupOk = RunStatus(Group.Value.CommonDir, MoveTemp(Group.Value.Files), OutErrorMessages, States, OutChangeset, OutBranchName);
+		if (!bGroupOk)
+		{
+			bResults = false;
+		}
+		else if (States.Num() > 0)
 		{
 			// Run a "fileinfo" command to update complementary status information of given files.
 			// (ie RevisionChangeset, RevisionHeadChangeset, RepSpec, LockedBy, LockedWhere)
 			// In case of "whole directory status", there is no explicit file in the group (it contains only the directory)
 			// => work on the list of files discovered by RunStatus()
-			bResults &= RunFileinfo(bWholeDirectory, bInUpdateHistory, InConcurrency, OutErrorMessages, States);
+			bResults &= RunFileinfo(bWholeDirectory, bInUpdateHistory, OutErrorMessages, States);
 		}
 		OutStates.Append(MoveTemp(States));
 	}
@@ -1311,7 +1066,7 @@ bool RunUpdateStatus(const TArray<FString>& InFiles, const bool bInUpdateHistory
 // cm cat revid:1230@rep:myrep@repserver:myserver:8084 --raw --file=Name124.tmp
 bool RunDumpToFile(const FString& InPathToPlasticBinary, const FString& InRevSpec, const FString& InDumpFileName)
 {
-	TRACE_CPUPROFILER_EVENT_SCOPE(FPlasticSourceControlUtils::RunDumpToFile);
+	TRACE_CPUPROFILER_EVENT_SCOPE(PlasticSourceControlUtils::RunDumpToFile);
 
 	int32	ReturnCode = 0;
 	FString Results;
@@ -1366,6 +1121,8 @@ FString DecodeXmlEntities(const FString& InString)
 		  <Repository>UE4PlasticPluginDev</Repository>
 		  <Server>localhost:8087</Server>
 		  <RepositorySpec>UE4PlasticPluginDev@localhost:8087</RepositorySpec>
+		  <Size>22356</Size>
+		  <Hash>zzuB6G9fbWz1md12+tvBxg==</Hash>
 		</Revision>
 		...
 		<Revision>
@@ -1379,6 +1136,8 @@ FString DecodeXmlEntities(const FString& InString)
 		  <Repository>UE4PlasticPluginDev</Repository>
 		  <Server>localhost:8087</Server>
 		  <RepositorySpec>UE4PlasticPluginDev@localhost:8087</RepositorySpec>
+		  <Size>22406</Size>
+		  <Hash>uR7NdDRAyKqADdyAqh67Rg==</Hash>
 		</Revision>
 
 	  </Revisions>
@@ -1392,7 +1151,7 @@ FString DecodeXmlEntities(const FString& InString)
 */
 static bool ParseHistoryResults(const bool bInUpdateHistory, const FXmlFile& InXmlResult, TArray<FPlasticSourceControlState>& InOutStates)
 {
-	TRACE_CPUPROFILER_EVENT_SCOPE(FPlasticSourceControlUtils::ParseHistoryResults);
+	TRACE_CPUPROFILER_EVENT_SCOPE(PlasticSourceControlUtils::ParseHistoryResults);
 
 	const FPlasticSourceControlProvider& Provider = FPlasticSourceControlModule::Get().GetProvider();
 	const FString RootRepSpec = FString::Printf(TEXT("%s@%s"), *Provider.GetRepositoryName(), *Provider.GetServerUrl());
@@ -1409,6 +1168,8 @@ static bool ParseHistoryResults(const bool bInUpdateHistory, const FXmlFile& InX
 	static const FString ChangesetNumber(TEXT("ChangesetNumber"));
 	static const FString Owner(TEXT("Owner"));
 	static const FString Comment(TEXT("Comment"));
+	static const FString Size(TEXT("Size"));
+	static const FString Hash(TEXT("Hash"));
 
 	const FXmlNode* RevisionHistoriesResultNode = InXmlResult.GetRootNode();
 	if (RevisionHistoriesResultNode == nullptr || RevisionHistoriesResultNode->GetTag() != RevisionHistoriesResult)
@@ -1526,9 +1287,16 @@ static bool ParseHistoryResults(const bool bInUpdateHistory, const FXmlFile& InX
 				{
 					SourceControlRevision->Branch = BranchNode->GetContent();
 				}
+				if (const FXmlNode* SizeNode = RevisionNode->FindChildNode(Size))
+				{
+					SourceControlRevision->FileSize = FCString::Atoi(*SizeNode->GetContent());
+				}
 
 				// Detect and skip more recent changesets on other branches (ie above the RevisionHeadChangeset)
-				if (SourceControlRevision->ChangesetNumber > InOutState.DepotRevisionChangeset)
+				// since we usually don't want to display changes from other branches in the History window...
+				// except in case of a merge conflict, where the Editor expects the tip of the "source (remote)" branch to be at the top of the history!
+				if (   (SourceControlRevision->ChangesetNumber > InOutState.DepotRevisionChangeset)
+					&& (SourceControlRevision->ChangesetNumber != InOutState.PendingMergeSourceChangeset))
 				{
 					InOutState.HeadBranch = SourceControlRevision->Branch;
 					InOutState.HeadAction = SourceControlRevision->Action;
@@ -1561,6 +1329,8 @@ static bool ParseHistoryResults(const bool bInUpdateHistory, const FXmlFile& InX
 // Run a Plastic "history" command and parse it's XML result.
 bool RunGetHistory(const bool bInUpdateHistory, TArray<FPlasticSourceControlState>& InOutStates, TArray<FString>& OutErrorMessages)
 {
+	TRACE_CPUPROFILER_EVENT_SCOPE(PlasticSourceControlUtils::RunGetHistory);
+
 	bool bResult = true;
 	FString Results;
 	FString Errors;
@@ -1572,6 +1342,20 @@ bool RunGetHistory(const bool bInUpdateHistory, TArray<FPlasticSourceControlStat
 	}
 	Parameters.Add(TEXT("--xml"));
 	Parameters.Add(TEXT("--encoding=\"utf-8\""));
+	const FPlasticSourceControlProvider& Provider = FPlasticSourceControlModule::Get().GetProvider();
+	if (Provider.GetPlasticScmVersion() >= PlasticSourceControlVersions::NewHistoryLimit)
+	{
+		if (bInUpdateHistory)
+		{
+			// --limit=0 will not limit the number of revisions, as stated by LimitNumberOfRevisionsInHistory
+			Parameters.Add(FString::Printf(TEXT("--limit=%d"), GetDefault<UPlasticSourceControlProjectSettings>()->LimitNumberOfRevisionsInHistory));
+		}
+		else
+		{
+			// when only searching for more recent changes on other branches, only the last revision is needed (to compare to the head of the current branch)
+			Parameters.Add(TEXT("--limit=1"));
+		}
+	}
 
 	TArray<FString> Files;
 	Files.Reserve(InOutStates.Num());
@@ -1588,18 +1372,22 @@ bool RunGetHistory(const bool bInUpdateHistory, TArray<FPlasticSourceControlStat
 	}
 	if (Files.Num() > 0)
 	{
-		bResult = RunCommandInternal(TEXT("history"), Parameters, Files, EConcurrency::Synchronous, Results, Errors);
+		bResult = RunCommand(TEXT("history"), Parameters, Files, Results, Errors);
 		if (bResult)
 		{
 			FXmlFile XmlFile;
 			{
-				TRACE_CPUPROFILER_EVENT_SCOPE(FPlasticSourceControlUtils::RunGetHistory::FXmlFile::LoadFile);
+				TRACE_CPUPROFILER_EVENT_SCOPE(PlasticSourceControlUtils::RunGetHistory::FXmlFile::LoadFile);
 				bResult = XmlFile.LoadFile(Results, EConstructMethod::ConstructFromBuffer);
 			}
 			if (bResult)
 			{
 				bResult = ParseHistoryResults(bInUpdateHistory, XmlFile, InOutStates);
 			}
+			else
+			{
+				UE_LOG(LogSourceControl, Error, TEXT("RunGetHistory: XML parse error '%s'"), *XmlFile.GetLastError())
+			}
 		}
 		if (!Errors.IsEmpty())
 		{
@@ -1611,6 +1399,139 @@ bool RunGetHistory(const bool bInUpdateHistory, TArray<FPlasticSourceControlStat
 }
 
 
+/* Parse results of the 'cm update --xml=tempfile.xml --encoding="utf-8"' command.
+ * 
+ * Results of the update command looks like that:
+<UpdatedItems>
+  <List>
+	<UpdatedItem>
+	  <Path>c:\Workspace\UE5PlasticPluginDev\Content\NewFolder\BP_CheckedOut.uasset</Path>
+	  <User>[REDACTED_EMAIL]</User>
+	  <Changeset>94</Changeset>
+	  <Date>2022-10-27T11:58:02+02:00</Date>
+	</UpdatedItem>
+  </List>
+</UpdatedItems>
+*/
+static bool ParseUpdateResults(const FXmlFile& InXmlResult, TArray<FString>& OutFiles)
+{
+	TRACE_CPUPROFILER_EVENT_SCOPE(PlasticSourceControlUtils::ParseUpdateResults);
+
+	static const FString UpdatedItems(TEXT("UpdatedItems"));
+	static const FString List(TEXT("List"));
+	static const FString UpdatedItem(TEXT("UpdatedItem"));
+	static const FString Path(TEXT("Path"));
+
+	const FXmlNode* UpdatedItemsNode = InXmlResult.GetRootNode();
+	if (UpdatedItemsNode == nullptr || UpdatedItemsNode->GetTag() != UpdatedItems)
+	{
+		return false;
+	}
+
+	const FXmlNode* ListNode = UpdatedItemsNode->FindChildNode(List);
+	if (ListNode == nullptr)
+	{
+		return false;
+	}
+
+	const TArray<FXmlNode*>& UpdatedItemNodes = ListNode->GetChildrenNodes();
+	for (const FXmlNode* UpdatedItemNode : UpdatedItemNodes)
+	{
+		if (const FXmlNode* PathNode = UpdatedItemNode->FindChildNode(Path))
+		{
+			FString Filename = PathNode->GetContent();
+			FPaths::NormalizeFilename(Filename);
+			OutFiles.Add(Filename);
+		}
+	}
+
+	return true;
+}
+
+
+/* Parse results of the 'cm partial update --report --machinereadable' command.
+ *
+ * Results of the update command looks like that:
+STAGE Plastic is updating your workspace. Wait a moment, please...
+STAGE Updated 63.01 KB of 63.01 KB (12 of 12 files to download / 16 of 21 operations to apply) /Content/Collections/SebSharedCollection.collection
+AD c:\Workspace\UE5PlasticPluginDev\Content\LevelPrototyping\Materials\MI_Solid_Red.uasset
+CH c:\Workspace\UE5PlasticPluginDev\Config\DefaultEditor.ini
+DE c:\Workspace\UE5PlasticPluginDev\Content\Collections\SebSharedCollection.collection
+*/
+static bool ParseUpdateResults(const TArray<FString>& InResults, TArray<FString>& OutFiles)
+{
+	TRACE_CPUPROFILER_EVENT_SCOPE(PlasticSourceControlUtils::ParseUpdateResultsString);
+
+	static const FString Stage = TEXT("STAGE ");
+	static const int32 PrefixLen = 3; // "XX " typically "CH ", "AD " or "DE "
+
+	for (const FString& Result : InResults)
+	{
+		if (Result.StartsWith(Stage))
+			continue;
+
+		FString Filename = Result.RightChop(PrefixLen);
+		FPaths::NormalizeFilename(Filename);
+		OutFiles.Add(Filename);
+	}
+
+	return true;
+}
+
+// Run a Plastic "update" command to sync the workspace and parse its XML results.
+bool RunUpdate(const TArray<FString>& InFiles, const bool bInIsPartialWorkspace, TArray<FString>& OutUpdatedFiles, TArray<FString>& OutErrorMessages)
+{
+	bool bResult = false;
+
+	TArray<FString> Parameters;
+	// Update specified directory to the head of the repository
+	// Detect special case for a partial checkout (CS:-1 in Gluon mode)!
+	if (!bInIsPartialWorkspace)
+	{
+		const FScopedTempFile TempFile;
+		TArray<FString> InfoMessages;
+		Parameters.Add(FString::Printf(TEXT("--xml=\"%s\""), *TempFile.GetFilename()));
+		Parameters.Add(TEXT("--encoding=\"utf-8\""));
+		Parameters.Add(TEXT("--last"));
+		Parameters.Add(TEXT("--dontmerge"));
+		bResult = PlasticSourceControlUtils::RunCommand(TEXT("update"), Parameters, TArray<FString>(), InfoMessages, OutErrorMessages);
+		if (bResult)
+		{
+			// Load and parse the result of the update command
+			FString Results;
+			if (FFileHelper::LoadFileToString(Results, *TempFile.GetFilename()))
+			{
+				FXmlFile XmlFile;
+				{
+					TRACE_CPUPROFILER_EVENT_SCOPE(PlasticSourceControlUtils::RunUpdate::FXmlFile::LoadFile);
+					bResult = XmlFile.LoadFile(Results, EConstructMethod::ConstructFromBuffer);
+				}
+				if (bResult)
+				{
+					bResult = ParseUpdateResults(XmlFile, OutUpdatedFiles);
+				}
+				else
+				{
+					UE_LOG(LogSourceControl, Error, TEXT("RunUpdate: XML parse error '%s'"), *XmlFile.GetLastError())
+				}
+			}
+		}
+	}
+	else
+	{
+		TArray<FString> Results;
+		Parameters.Add(TEXT("--report"));
+		Parameters.Add(TEXT("--machinereadable"));
+		bResult = PlasticSourceControlUtils::RunCommand(TEXT("partial update"), Parameters, InFiles, Results, OutErrorMessages);
+		if (bResult)
+		{
+			bResult = ParseUpdateResults(Results, OutUpdatedFiles);
+		}
+	}
+
+	return bResult;
+}
+
 /**
  * Parse results of the 'cm status --changelists --controlledchanged --xml --encoding="utf-8"' command.
  *
@@ -1657,7 +1578,7 @@ bool RunGetHistory(const bool bInUpdateHistory, TArray<FPlasticSourceControlStat
 */
 static bool ParseChangelistsResults(const FXmlFile& InXmlResult, TArray<FPlasticSourceControlChangelistState>& OutChangelistsStates, TArray<TArray<FPlasticSourceControlState>>& OutCLFilesStates)
 {
-	TRACE_CPUPROFILER_EVENT_SCOPE(FPlasticSourceControlUtils::ParseChangelistsResults);
+	TRACE_CPUPROFILER_EVENT_SCOPE(PlasticSourceControlUtils::ParseChangelistsResults);
 
 	static const FString StatusOutput(TEXT("StatusOutput"));
 	static const FString WkConfigType(TEXT("WkConfigType"));
@@ -1741,8 +1662,11 @@ static bool ParseChangelistsResults(const FXmlFile& InXmlResult, TArray<FPlastic
 	return true;
 }
 
-bool RunGetChangelists(const EConcurrency::Type InConcurrency, TArray<FPlasticSourceControlChangelistState>& OutChangelistsStates, TArray<TArray<FPlasticSourceControlState>>& OutCLFilesStates, TArray<FString>& OutErrorMessages)
+// Run a Plastic "status --changelist --xml" and parse its XML result.
+bool RunGetChangelists(TArray<FPlasticSourceControlChangelistState>& OutChangelistsStates, TArray<TArray<FPlasticSourceControlState>>& OutCLFilesStates, TArray<FString>& OutErrorMessages)
 {
+	TRACE_CPUPROFILER_EVENT_SCOPE(PlasticSourceControlUtils::RunGetChangelists);
+
 	FString Results;
 	FString Errors;
 	TArray<FString> Parameters;
@@ -1751,18 +1675,22 @@ bool RunGetChangelists(const EConcurrency::Type InConcurrency, TArray<FPlasticSo
 	Parameters.Add(TEXT("--noheader"));
 	Parameters.Add(TEXT("--xml"));
 	Parameters.Add(TEXT("--encoding=\"utf-8\""));
-	bool bResult = PlasticSourceControlUtils::RunCommandInternal(TEXT("status"), Parameters, TArray<FString>(), InConcurrency, Results, Errors);
+	bool bResult = RunCommand(TEXT("status"), Parameters, TArray<FString>(), Results, Errors);
 	if (bResult)
 	{
 		FXmlFile XmlFile;
 		{
-			TRACE_CPUPROFILER_EVENT_SCOPE(FPlasticSourceControlUtils::RunGetHistory::FXmlFile::LoadFile);
+			TRACE_CPUPROFILER_EVENT_SCOPE(PlasticSourceControlUtils::RunGetChangelists::FXmlFile::LoadFile);
 			bResult = XmlFile.LoadFile(Results, EConstructMethod::ConstructFromBuffer);
 		}
 		if (bResult)
 		{
 			bResult = ParseChangelistsResults(XmlFile, OutChangelistsStates, OutCLFilesStates);
 		}
+		else
+		{
+			UE_LOG(LogSourceControl, Error, TEXT("RunGetChangelists: XML parse error '%s'"), *XmlFile.GetLastError())
+		}
 	}
 	if (!Errors.IsEmpty())
 	{
@@ -1772,6 +1700,246 @@ bool RunGetChangelists(const EConcurrency::Type InConcurrency, TArray<FPlasticSo
 	return bResult;
 }
 
+// Parse the one letter file status in front of each line of the 'cm diff sh:<ShelveId>'
+EWorkspaceState ParseShelveFileStatus(const TCHAR InFileStatus)
+{
+	if (InFileStatus == 'A') // Added
+	{
+		return EWorkspaceState::Added;
+	}
+	else if (InFileStatus == 'D') // Deleted
+	{
+		return EWorkspaceState::Deleted;
+	}
+	else if (InFileStatus == 'C') // Changed (CheckedOut or not)
+	{
+		return EWorkspaceState::CheckedOut;
+	}
+	else if (InFileStatus == 'M') // Moved/Renamed (or Locally Moved)
+	{
+		return EWorkspaceState::Moved;
+	}
+	else
+	{
+		UE_LOG(LogSourceControl, Warning, TEXT("Unknown file status '%c'"), InFileStatus);
+		return EWorkspaceState::Unknown;
+	}
+}
+
+void AddShelvedFileToChangelist(FPlasticSourceControlChangelistState& InOutChangelistsState, FString&& InFilename, EWorkspaceState InShelveStatus)
+{
+	TSharedRef<FPlasticSourceControlState, ESPMode::ThreadSafe> ShelveState = MakeShared<FPlasticSourceControlState>(MoveTemp(InFilename), InShelveStatus);
+
+	// Add one revision to be able to fetch the shelved file for diff, if it's not marked for deletion.
+	if (InShelveStatus != EWorkspaceState::Deleted)
+	{
+		const TSharedRef<FPlasticSourceControlRevision, ESPMode::ThreadSafe> SourceControlRevision = MakeShared<FPlasticSourceControlRevision>();
+		SourceControlRevision->State = &ShelveState.Get();
+		SourceControlRevision->Filename = ShelveState->GetFilename();
+		SourceControlRevision->ShelveId = InOutChangelistsState.ShelveId;
+		SourceControlRevision->ChangesetNumber = InOutChangelistsState.ShelveId; // Note: for display in the diff window only
+		SourceControlRevision->Date = InOutChangelistsState.ShelveDate; // Note: not yet used for display as of UE5.1.1
+
+		ShelveState->History.Add(SourceControlRevision);
+	}
+
+	// In case of a Moved file, it would appear twice in the list, so overwrite it if already in
+	if (FSourceControlStateRef* ExistingShelveState = InOutChangelistsState.ShelvedFiles.FindByPredicate(
+		[&ShelveState](const FSourceControlStateRef& State)
+		{
+			return State->GetFilename().Equals(ShelveState->GetFilename());
+		}))
+	{
+		*ExistingShelveState = MoveTemp(ShelveState);
+	}
+	else
+	{
+		InOutChangelistsState.ShelvedFiles.Add(MoveTemp(ShelveState));
+	}
+}
+
+
+/**
+ * Parse results of the 'cm diff sh:<ShelveId>' command.
+ *
+ * Results of the diff command looks like that:
+C "Content\NewFolder\BP_CheckedOut.uasset"
+C "Content\NewFolder\BP_Renamed.uasset"
+A "Content\NewFolder\BP_ControlledUnchanged.uasset"
+D "Content\NewFolder\BP_Changed.uasset"
+M "Content\NewFolder\BP_ControlledUnchanged.uasset" "Content\NewFolder\BP_Renamed.uasset"
+*/
+bool ParseShelveDiffResults(const FString InWorkingDirectory, TArray<FString>&& InResults, FPlasticSourceControlChangelistState& InOutChangelistsState)
+{
+	bool bCommandSuccessful = true;
+
+	InOutChangelistsState.ShelvedFiles.Reset(InResults.Num());
+	for (FString& Result : InResults)
+	{
+		EWorkspaceState ShelveStatus = ParseShelveFileStatus(Result[0]);
+		
+		// Remove outer double quotes
+		Result.MidInline(3, Result.Len() - 4, false);
+		if (ShelveStatus == EWorkspaceState::Moved)
+		{
+			// Search for the inner double quotes in the middle of "Content/Source.uasset" "Content/Destination.uasset" to keep only the destination filename
+			int32 RenameIndex;
+			if (Result.FindLastChar(TEXT('"'), RenameIndex))
+			{
+				Result.RightChopInline(RenameIndex + 1);
+			}
+		}
+		
+		if (ShelveStatus != EWorkspaceState::Unknown && !Result.IsEmpty())
+		{
+			FString AbsoluteFilename = FPaths::ConvertRelativePathToFull(InWorkingDirectory, MoveTemp(Result));
+			AddShelvedFileToChangelist(InOutChangelistsState, MoveTemp(AbsoluteFilename), ShelveStatus);
+		}
+		else
+		{
+			bCommandSuccessful = false;
+		}
+	}
+
+	return bCommandSuccessful;
+}
+
+/**
+ * Run for each shelve a "diff sh:<ShelveId>" and parse their result to list their files.
+ * @param	InOutChangelistsStates	The list of changelists, filled with their shelved files
+ * @param	OutErrorMessages		Any errors (from StdErr) as an array per-line
+ */
+bool RunGetShelveFiles(TArray<FPlasticSourceControlChangelistState>& InOutChangelistsStates, TArray<FString>& OutErrorMessages)
+{
+	bool bCommandSuccessful = true;
+
+	const FString& WorkingDirectory = FPlasticSourceControlModule::Get().GetProvider().GetPathToWorkspaceRoot();
+
+	for (FPlasticSourceControlChangelistState& ChangelistState : InOutChangelistsStates)
+	{
+		if (ChangelistState.ShelveId != ISourceControlState::INVALID_REVISION)
+		{
+			TArray<FString> Results;
+			TArray<FString> Parameters;
+			Parameters.Add(FString::Printf(TEXT("sh:%d"), ChangelistState.ShelveId));
+			const bool bDiffSuccessful = PlasticSourceControlUtils::RunCommand(TEXT("diff"), Parameters, TArray<FString>(), Results, OutErrorMessages);
+			if (bDiffSuccessful)
+			{
+				bCommandSuccessful = ParseShelveDiffResults(WorkingDirectory, MoveTemp(Results), ChangelistState);
+			}
+		}
+	}
+
+	return bCommandSuccessful;
+}
+
+/**
+ * Parse results of the 'cm find "shelves where owner='me'" --xml --encoding="utf-8"' command.
+ *
+ * Find shelves with comments starting like "ChangelistXXX: " and matching an existing Changelist number XXX
+ *
+ * Results of the find command looks like the following; note the "Changelist67: " prefix of the comment:
+<?xml version="1.0" encoding="utf-8" ?>
+<PLASTICQUERY>
+  <SHELVE>
+	<ID>1376</ID>
+	<SHELVEID>9</SHELVEID>
+    <COMMENT>Changelist67: test by Sebastien</COMMENT>
+	<DATE>2022-06-30T16:39:55+02:00</DATE>
+    <OWNER>[REDACTED_EMAIL]</OWNER>
+    <REPOSITORY>UE5PlasticPluginDev</REPOSITORY>
+    <REPNAME>UE5PlasticPluginDev</REPNAME>
+    <REPSERVER>test@cloud</REPSERVER>
+	<PARENT>45</PARENT>
+	<GUID>8fbefbcc-81a7-4b81-9b99-b51f4873d09f</GUID>
+  </SHELVE>
+  [...]
+</PLASTICQUERY>
+*/
+static bool ParseShelvesResults(const FXmlFile& InXmlResult, TArray<FPlasticSourceControlChangelistState>& InOutChangelistsStates)
+{
+	static const FString PlasticQuery(TEXT("PLASTICQUERY"));
+	static const FString Shelve(TEXT("SHELVE"));
+	static const FString ShelveId(TEXT("SHELVEID"));
+	static const FString Date(TEXT("DATE"));
+	static const FString Comment(TEXT("COMMENT"));
+
+	const FString& WorkingDirectory = FPlasticSourceControlModule::Get().GetProvider().GetPathToWorkspaceRoot();
+
+	const FXmlNode* PlasticQueryNode = InXmlResult.GetRootNode();
+	if (PlasticQueryNode == nullptr || PlasticQueryNode->GetTag() != PlasticQuery)
+	{
+		return false;
+	}
+
+	const TArray<FXmlNode*>& ShelvesNodes = PlasticQueryNode->GetChildrenNodes();
+	for (const FXmlNode* ShelveNode : ShelvesNodes)
+	{
+		check(ShelveNode);
+		const FXmlNode* ShelveIdNode = ShelveNode->FindChildNode(ShelveId);
+		const FXmlNode* CommentNode = ShelveNode->FindChildNode(Comment);
+		if (ShelveIdNode == nullptr || CommentNode == nullptr)
+		{
+			continue;
+		}
+
+		const FString& ShelveIdString = ShelveIdNode->GetContent();
+		const FString& CommentString = CommentNode->GetContent();
+
+		// Search if there is a changelist matching the shelve (that is, a shelve with a comment starting with "ChangelistXXX: ")
+		for (FPlasticSourceControlChangelistState& ChangelistState : InOutChangelistsStates)
+		{
+			FPlasticSourceControlChangelistRef Changelist = StaticCastSharedRef<FPlasticSourceControlChangelist>(ChangelistState.GetChangelist());
+			const FString ChangelistPrefix = FString::Printf(TEXT("Changelist%s: "), *Changelist->GetName());
+			if (CommentString.StartsWith(ChangelistPrefix))
+			{
+				ChangelistState.ShelveId = FCString::Atoi(*ShelveIdString);
+
+				if (const FXmlNode* DateNode = ShelveNode->FindChildNode(Date))
+				{
+					const FString& DateIso = DateNode->GetContent();
+					FDateTime::ParseIso8601(*DateIso, ChangelistState.ShelveDate);
+				}
+			}
+		}
+	}
+
+	return true;
+}
+
+// Run find "shelves where owner='me'" and for each shelve matching a changelist a "diff sh:<ShelveId>" and parse their results.
+bool RunGetShelves(TArray<FPlasticSourceControlChangelistState>& InOutChangelistsStates, TArray<FString>& OutErrorMessages)
+{
+	bool bCommandSuccessful;
+
+	FString Results;
+	FString Errors;
+	TArray<FString> Parameters;
+	Parameters.Add(TEXT("\"shelves where owner = 'me'\""));
+	Parameters.Add(TEXT("--xml"));
+	Parameters.Add(TEXT("--encoding=\"utf-8\""));
+	bCommandSuccessful = PlasticSourceControlUtils::RunCommand(TEXT("find"), Parameters, TArray<FString>(), Results, Errors);
+	if (bCommandSuccessful)
+	{
+		FXmlFile XmlFile;
+		bCommandSuccessful = XmlFile.LoadFile(Results, EConstructMethod::ConstructFromBuffer);
+		if (bCommandSuccessful)
+		{
+			bCommandSuccessful = ParseShelvesResults(XmlFile, InOutChangelistsStates);
+			if (bCommandSuccessful)
+			{
+				bCommandSuccessful = RunGetShelveFiles(InOutChangelistsStates, OutErrorMessages);
+			}
+		}
+	}
+	if (!Errors.IsEmpty())
+	{
+		OutErrorMessages.Add(MoveTemp(Errors));
+	}
+
+	return bCommandSuccessful;
+}
+
 bool UpdateCachedStates(TArray<FPlasticSourceControlState>&& InStates)
 {
 	FPlasticSourceControlProvider& Provider = FPlasticSourceControlModule::Get().GetProvider();
diff --git a/Engine/Plugins/Developer/PlasticSourceControl/Source/PlasticSourceControl/Private/PlasticSourceControlUtils.h b/Engine/Plugins/Developer/PlasticSourceControl/Source/PlasticSourceControl/Private/PlasticSourceControlUtils.h
--- a/Engine/Plugins/Developer/PlasticSourceControl/Source/PlasticSourceControl/Private/PlasticSourceControlUtils.h
+++ b/Engine/Plugins/Developer/PlasticSourceControl/Source/PlasticSourceControl/Private/PlasticSourceControlUtils.h
@@ -3,7 +3,6 @@
 #pragma once
 
 #include "CoreMinimal.h"
-#include "ISourceControlProvider.h"
 #include "PlasticSourceControlRevision.h"
 
 class FPlasticSourceControlChangelistState;
@@ -11,26 +10,7 @@ class FPlasticSourceControlCommand;
 class FPlasticSourceControlState;
 struct FSoftwareVersion;
 
-/**
- * Helper struct for maintaining temporary files for passing to commands
- */
-class FScopedTempFile
-{
-public:
-	/** Constructor - open & write string to temp file */
-	explicit FScopedTempFile(const FString& InText);
-	explicit FScopedTempFile(const FText& InText);
-
-	/** Destructor - delete temp file */
-	~FScopedTempFile();
-
-	/** Get the filename of this temp file - empty if it failed to be created */
-	const FString& GetFilename() const;
-
-private:
-	/** The filename we are writing to */
-	FString Filename;
-};
+enum class EWorkspaceState;
 
 namespace PlasticSourceControlUtils
 {
@@ -42,31 +22,34 @@ const struct FSoftwareVersion& GetOldestSupportedPlasticScmVersion();
  */
 FString FindPlasticBinaryPath();
 
-/**
- * Launch the Plastic SCM "shell" command to run it in background.
- * @param	InPathToPlasticBinary	The path to the Plastic binary
- * @param	InWorkspaceRoot			The workspace from where to run the command - usually the Game directory
- * @returns true if the command succeeded and returned no errors
- */
-bool LaunchBackgroundPlasticShell(const FString& InPathToPlasticBinary, const FString& InWorkspaceRoot);
-
-/** Terminate the background 'cm shell' process and associated pipes */
-void Terminate();
-
 /**
  * Find the root of the Plastic workspace, looking from the GameDir and upward in its parent directories
  * @param InPathToGameDir		The path to the Game Directory
  * @param OutWorkspaceRoot		The path to the root directory of the Plastic workspace if found, else the path to the GameDir
- * @returns true if the command succeeded and returned no errors
+ * @returns true if the command succeeded
  */
-bool FindRootDirectory(const FString& InPathToGameDir, FString& OutWorkspaceRoot);
+bool GetWorkspacePath(const FString& InPathToGameDir, FString& OutWorkspaceRoot);
 
 /**
- * Get Plastic SCM cli version
+ * Get Plastic SCM CLI version
  * @param	OutCliVersion		Version of the Plastic SCM Command Line Interface tool
+ * @returns true if the command succeeded
 */
 bool GetPlasticScmVersion(FSoftwareVersion& OutPlasticScmVersion);
 
+/**
+ * Get Plastic SCM CLI location
+ * @param	OutCmLocation		Path to the "cm" executable
+ * @returns true if the command succeeded
+*/
+bool GetCmLocation(FString& OutCmLocation);
+
+/**
+ * Checks weather Plastic SCM is configured to set files as read-only on update & checkin
+ * @returns true if SetFilesAsReadOnly is enabled in client.conf
+*/
+bool GetConfigSetFilesAsReadOnly();
+
 /**
  * Get Plastic SCM current user
  * @param	OutUserName			Name of the Plastic SCM user configured globally
@@ -91,6 +74,18 @@ bool GetWorkspaceName(const FString& InWorkspaceRoot, FString& OutWorkspaceName,
  */
 bool GetWorkspaceInformation(int32& OutChangeset, FString& OutRepositoryName, FString& OutServerUrl, FString& OutBranchName, TArray<FString>& OutErrorMessages);
 
+/**
+ * Get Plastic repository name, server URL, and branch name
+ *
+ * Note: this is a local fast variant, not making network call to the server.
+ * 
+ * @param	OutBranchName		Name of the current checked-out branch
+ * @param	OutRepositoryName	Name of the repository of the current workspace
+ * @param	OutServerUrl		URL/Port of the server of the repository
+ * @param	OutErrorMessages	Any errors (from StdErr) as an array per-line
+ */
+bool GetWorkspaceInfo(FString& OutBranchName, FString& OutRepositoryName, FString& OutServerUrl, TArray<FString>& OutErrorMessages);
+
 /**
  * Use the Project Settings to replace Plastic SCM full username/e-mail by a shorter version for display.
  *
@@ -101,33 +96,41 @@ bool GetWorkspaceInformation(int32& OutChangeset, FString& OutRepositoryName, FS
 FString UserNameToDisplayName(const FString& InUserName);
 
 /**
- * Run a Plastic command - output is a string TArray.
+ * Run a Plastic command - the result is the output of cm, as a multi-line string.
+ *
+ * @param	InCommand			The Plastic command - e.g. commit
+ * @param	InParameters		The parameters to the Plastic command
+ * @param	InFiles				The files to be operated on
+ * @param	OutResults			The results (from StdOut) as a multi-line string.
+ * @param	OutErrors			Any errors (from StdErr) as a multi-line string.
+ * @returns true if the command succeeded and returned no errors
+ */
+bool RunCommand(const FString& InCommand, const TArray<FString>& InParameters, const TArray<FString>& InFiles, FString& OutResults, FString& OutErrors);
+
+/**
+ * Run a Plastic command - the result is parsed in an array of strings.
  *
  * @param	InCommand			The Plastic command - e.g. commit
  * @param	InParameters		The parameters to the Plastic command
  * @param	InFiles				The files to be operated on
- * @param	InConcurrency		Is the command running in the background, or blocking the main thread
  * @param	OutResults			The results (from StdOut) as an array per-line
  * @param	OutErrorMessages	Any errors (from StdErr) as an array per-line
  * @returns true if the command succeeded and returned no errors
  */
-bool RunCommand(const FString& InCommand, const TArray<FString>& InParameters, const TArray<FString>& InFiles, const EConcurrency::Type InConcurrency, TArray<FString>& OutResults, TArray<FString>& OutErrorMessages);
-// Run a Plastic command - output is a string.
-bool RunCommandInternal(const FString& InCommand, const TArray<FString>& InParameters, const TArray<FString>& InFiles, const EConcurrency::Type InConcurrency, FString& OutResults, FString& OutErrors);
+bool RunCommand(const FString& InCommand, const TArray<FString>& InParameters, const TArray<FString>& InFiles, TArray<FString>& OutResults, TArray<FString>& OutErrorMessages);
 
 /**
  * Run a Plastic "status" command and parse it.
  *
  * @param	InFiles				The files to be operated on
  * @param	bInUpdateHistory	If getting the history of files, force execute the fileinfo command required to do get RepSpec of xlinks (history view or visual diff)
- * @param	InConcurrency		Is the command running in the background, or blocking the main thread
  * @param	OutErrorMessages	Any errors (from StdErr) as an array per-line
  * @param	OutStates			States of the files
  * @param	OutChangeset		The current Changeset Number
  * @param	OutBranchName		Name of the current checked-out branch
  * @returns true if the command succeeded and returned no errors
  */
-bool RunUpdateStatus(const TArray<FString>& InFiles, const bool bInUpdateHistory, const EConcurrency::Type InConcurrency, TArray<FString>& OutErrorMessages, TArray<FPlasticSourceControlState>& OutStates, int32& OutChangeset, FString& OutBranchName);
+bool RunUpdateStatus(const TArray<FString>& InFiles, const bool bInUpdateHistory, TArray<FString>& OutErrorMessages, TArray<FPlasticSourceControlState>& OutStates, int32& OutChangeset, FString& OutBranchName);
 
 /**
  * Run a Plastic "cat" command to dump the binary content of a revision into a file.
@@ -140,7 +143,7 @@ bool RunUpdateStatus(const TArray<FString>& InFiles, const bool bInUpdateHistory
 bool RunDumpToFile(const FString& InPathToPlasticBinary, const FString& InRevSpec, const FString& InDumpFileName);
 
 /**
- * Run Plastic "history" and "log" commands and parse their results.
+ * Run Plastic "history" and "log" commands and parse their XML results.
  *
  * @param	bInUpdateHistory	If getting the history of files, versus only checking the heads of branches to detect newer commits
  * @param	InOutStates			The file states to update with the history
@@ -149,13 +152,37 @@ bool RunDumpToFile(const FString& InPathToPlasticBinary, const FString& InRevSpe
 bool RunGetHistory(const bool bInUpdateHistory, TArray<FPlasticSourceControlState>& InOutStates, TArray<FString>& OutErrorMessages);
 
 /**
- * Run a Plastic "status --changelist --xml" and parse its result.
- * @param	InConcurrency
+ * Run a Plastic "update" command to sync the workspace and parse its XML results.
+ *
+ * @param	InFiles					The files or paths to sync
+ * @param	bInIsPartialWorkspace	Whether running on a partial/gluon or regular/full workspace
+ * @param	OutUpdatedFiles			The files that where updated
+ * @param	OutErrorMessages		Any errors (from StdErr) as an array per-line
+ */
+bool RunUpdate(const TArray<FString>& InFiles, const bool bInIsPartialWorkspace, TArray<FString>& OutUpdatedFiles, TArray<FString>& OutErrorMessages);
+
+/**
+ * Run a Plastic "status --changelist --xml" and parse its XML result.
  * @param	OutChangelistsStates	The list of changelists (without their files)
  * @param	OutCLFilesStates		The list of files per changelist
  * @param	OutErrorMessages		Any errors (from StdErr) as an array per-line
  */
-bool RunGetChangelists(const EConcurrency::Type InConcurrency, TArray<FPlasticSourceControlChangelistState>& OutChangelistsStates, TArray<TArray<FPlasticSourceControlState>>& OutCLFilesStates, TArray<FString>& OutErrorMessages);
+bool RunGetChangelists(TArray<FPlasticSourceControlChangelistState>& OutChangelistsStates, TArray<TArray<FPlasticSourceControlState>>& OutCLFilesStates, TArray<FString>& OutErrorMessages);
+
+/**
+ * Run find "shelves where owner='me'" and for each shelve matching a changelist a "diff sh:<ShelveId>" and parse their results.
+ * @param	InOutChangelistsStates	The list of changelists, filled with their shelved files
+ * @param	OutErrorMessages		Any errors (from StdErr) as an array per-line
+ */
+bool RunGetShelves(TArray<FPlasticSourceControlChangelistState>& InOutChangelistsStates, TArray<FString>& OutErrorMessages);
+
+/**
+ * Add a file to the shelve associated with a changelist.
+ * @param	InOutChangelistsState	The changelist to add the file to
+ * @param	InFilename				The file to add to the shelve
+ * @param	InShelveStatus			The status of the file
+ */
+void AddShelvedFileToChangelist(FPlasticSourceControlChangelistState& InOutChangelistsState, FString&& InFilename, EWorkspaceState InShelveStatus);
 
 /**
  * Helper function for various commands to update cached states.
diff --git a/Engine/Plugins/Developer/PlasticSourceControl/Source/PlasticSourceControl/Private/PlasticSourceControlVersions.h b/Engine/Plugins/Developer/PlasticSourceControl/Source/PlasticSourceControl/Private/PlasticSourceControlVersions.h
new file mode 100644
--- /dev/null
+++ b/Engine/Plugins/Developer/PlasticSourceControl/Source/PlasticSourceControl/Private/PlasticSourceControlVersions.h
@@ -0,0 +1,32 @@
+// Copyright Epic Games, Inc. All Rights Reserved.
+
+#pragma once
+
+#include "CoreMinimal.h"
+#include "SoftwareVersion.h"
+
+/**
+ * List of key versions of Plastic SCM enabling new features for "cm" CLI.
+ *
+ * Plastic SCM version strings in the form "X.Y.Z.C", ie Major.Minor.Patch.Changeset (as returned by GetPlasticScmVersion)
+*/
+namespace PlasticSourceControlVersions
+{
+	// Oldest supported version of "cm" (from 2021/01/05)
+	// 9.0.16.4839 cm changelist 'persistent' flag now contain a '--' prefix.
+	// https://www.plasticscm.com/download/releasenotes/9.0.16.4839
+	static FSoftwareVersion OldestSupported(TEXT("9.0.16.4839"));
+
+	// 11.0.16.7248 add support for --descriptionfile for multi-line descriptions and support for special characters
+	// https://www.plasticscm.com/download/releasenotes/11.0.16.7248
+	static const FSoftwareVersion NewChangelistFileArgs(TEXT("11.0.16.7248"));
+
+	// 11.0.16.7504 add support for 'cm shelveset apply' to only unshelved the desired changes from a shelve
+	// https://www.plasticscm.com/download/releasenotes/11.0.16.7504
+	static const FSoftwareVersion ShelvesetApplySelection(TEXT("11.0.16.7504"));
+
+	// 11.0.16.7608 add support for history --limit. It displays the N last revisions of the specified items.
+	// https://www.plasticscm.com/download/releasenotes/11.0.16.7608
+	static const FSoftwareVersion NewHistoryLimit(TEXT("11.0.16.7608"));
+
+} // namespace PlasticSourceControlVersions
diff --git a/Engine/Plugins/Developer/PlasticSourceControl/Source/PlasticSourceControl/Private/SPlasticSourceControlSettings.cpp b/Engine/Plugins/Developer/PlasticSourceControl/Source/PlasticSourceControl/Private/SPlasticSourceControlSettings.cpp
--- a/Engine/Plugins/Developer/PlasticSourceControl/Source/PlasticSourceControl/Private/SPlasticSourceControlSettings.cpp
+++ b/Engine/Plugins/Developer/PlasticSourceControl/Source/PlasticSourceControl/Private/SPlasticSourceControlSettings.cpp
@@ -29,7 +29,7 @@
 
 void SPlasticSourceControlSettings::Construct(const FArguments& InArgs)
 {
-	FSlateFontInfo Font = FAppStyle::GetFontStyle(TEXT("SourceControl.LoginWindow.Font"));
+	const FSlateFontInfo Font = FAppStyle::GetFontStyle(TEXT("SourceControl.LoginWindow.Font"));
 
 	bAutoCreateIgnoreFile = CanAutoCreateIgnoreFile();
 	bAutoInitialCommit = true;
@@ -55,7 +55,13 @@ void SPlasticSourceControlSettings::Construct(const FArguments& InArgs)
 			.ToolTipText(LOCTEXT("PlasticVersions_Tooltip", "Plastic SCM and Plugin versions"))
 			+SHorizontalBox::Slot()
 			.FillWidth(1.0f)
-			.HAlign(HAlign_Center)
+			[
+				SNew(STextBlock)
+				.Text(LOCTEXT("PlasticVersions", "Plastic SCM version"))
+				.Font(Font)
+			]
+			+SHorizontalBox::Slot()
+			.FillWidth(2.0f)
 			[
 				SNew(STextBlock)
 				.Text(this, &SPlasticSourceControlSettings::GetVersions)
@@ -85,7 +91,7 @@ void SPlasticSourceControlSettings::Construct(const FArguments& InArgs)
 			.FillWidth(1.0f)
 			[
 				SNew(STextBlock)
-				.Text(LOCTEXT("PathLabel", "Plastic SCM Path to cm"))
+				.Text(LOCTEXT("PathLabel", "Plastic SCM path to cm"))
 				.Font(Font)
 			]
 			+SHorizontalBox::Slot()
@@ -155,7 +161,7 @@ void SPlasticSourceControlSettings::Construct(const FArguments& InArgs)
 		// Explanation text
 		+SVerticalBox::Slot()
 		.FillHeight(1.0f)
-		.Padding(2.0f)
+		.Padding(2.0f, 5.0f)
 		[
 			SNew(STextBlock)
 			.Visibility(this, &SPlasticSourceControlSettings::CanInitializePlasticWorkspace)
@@ -306,7 +312,7 @@ void SPlasticSourceControlSettings::Construct(const FArguments& InArgs)
 			.OnCheckStateChanged(this, &SPlasticSourceControlSettings::OnCheckedUpdateStatusOtherBranches)
 			[
 				SNew(STextBlock)
-				.Text(LOCTEXT("UpdateStatusOtherBranches", "Update status also check branch history."))
+				.Text(LOCTEXT("UpdateStatusOtherBranches", "Update Status also checks history to detect changes on other branches."))
 				.Font(Font)
 			]
 		]
@@ -396,7 +402,7 @@ void SPlasticSourceControlSettings::OnBinaryPathTextCommited(const FText& InText
 FText SPlasticSourceControlSettings::GetVersions() const
 {
 	const FPlasticSourceControlProvider& Provider = FPlasticSourceControlModule::Get().GetProvider();
-	return FText::FromString(TEXT("Plastic SCM ") + Provider.GetPlasticScmVersion().String + TEXT(" (plugin v") + Provider.GetPluginVersion() + TEXT(")"));
+	return FText::FromString(Provider.GetPlasticScmVersion().String + TEXT(" (plugin v") + Provider.GetPluginVersion() + TEXT(")"));
 }
 
 FText SPlasticSourceControlSettings::GetPathToWorkspaceRoot() const
@@ -664,7 +670,7 @@ FReply SPlasticSourceControlSettings::OnClickedAddIgnoreFile() const
 		Parameters.Add(TEXT("-R"));
 		TArray<FString> Files;
 		Files.Add(TEXT("ignore.conf"));
-		PlasticSourceControlUtils::RunCommand(TEXT("add"), Parameters, Files, EConcurrency::Synchronous, InfoMessages, ErrorMessages);
+		PlasticSourceControlUtils::RunCommand(TEXT("add"), Parameters, Files, InfoMessages, ErrorMessages);
 	}
 	return FReply::Handled();
 }
@@ -721,7 +727,7 @@ const FString SPlasticSourceControlSettings::GetIgnoreFileName() const
 /** Create a standard "ignore.conf" file with common patterns for a typical Blueprint & C++ project */
 bool SPlasticSourceControlSettings::CreateIgnoreFile() const
 {
-	const FString IgnoreFileContent = TEXT("Binaries\nBuild\nDerivedDataCache\nIntermediate\nSaved\nScript\nenc_temp_folder\n.idea\n.vscode\n.vs\n*.VC.db\n*.opensdf\n*.opendb\n*.sdf\n*.sln\n*.suo\n*.xcodeproj\n*.xcworkspace");
+	const FString IgnoreFileContent = TEXT("Binaries\nBuild\nDerivedDataCache\nIntermediate\nSaved\nScript\nenc_temp_folder\n.idea\n.vscode\n.vs\n.vsconfig\n.ignore\n*.VC.db\n*.opensdf\n*.opendb\n*.sdf\n*.sln\n*.suo\n*.code-workspace\n*.xcodeproj\n*.xcworkspace");
 	return FFileHelper::SaveStringToFile(IgnoreFileContent, *GetIgnoreFileName(), FFileHelper::EEncodingOptions::ForceUTF8WithoutBOM);
 }
 
@@ -736,7 +742,7 @@ TArray<FString> SPlasticSourceControlSettings::GetProjectFiles() const
 	{
 		ProjectFiles.Add(FPaths::ConvertRelativePathToFull(FPaths::GameSourceDir()));
 	}
-	if (bAutoCreateIgnoreFile)
+	if (FPaths::FileExists(GetIgnoreFileName()))
 	{
 		ProjectFiles.Add(GetIgnoreFileName());
 	}
diff --git a/Engine/Plugins/Developer/PlasticSourceControl/Source/PlasticSourceControl/Private/ScopedTempFile.cpp b/Engine/Plugins/Developer/PlasticSourceControl/Source/PlasticSourceControl/Private/ScopedTempFile.cpp
new file mode 100644
--- /dev/null
+++ b/Engine/Plugins/Developer/PlasticSourceControl/Source/PlasticSourceControl/Private/ScopedTempFile.cpp
@@ -0,0 +1,48 @@
+// Copyright Epic Games, Inc. All Rights Reserved.
+
+#include "ScopedTempFile.h"
+
+#include "Misc/FileHelper.h"
+#include "Misc/Paths.h"
+#include "HAL/PlatformFileManager.h"
+
+#include "ISourceControlModule.h" // LogSourceControl
+
+FScopedTempFile::FScopedTempFile()
+{
+	Filename = FPaths::CreateTempFilename(*FPaths::ProjectLogDir(), TEXT("Plastic-Temp"), TEXT(".txt"));
+}
+
+FScopedTempFile::FScopedTempFile(const FString& InText)
+{
+	Filename = FPaths::CreateTempFilename(*FPaths::ProjectLogDir(), TEXT("Plastic-Temp"), TEXT(".txt"));
+	if (!FFileHelper::SaveStringToFile(InText, *Filename, FFileHelper::EEncodingOptions::ForceUTF8WithoutBOM))
+	{
+		UE_LOG(LogSourceControl, Error, TEXT("Failed to write to temp file: %s"), *Filename);
+	}
+}
+
+FScopedTempFile::FScopedTempFile(const FText& InText)
+{
+	Filename = FPaths::CreateTempFilename(*FPaths::ProjectLogDir(), TEXT("Plastic-Temp"), TEXT(".txt"));
+	if (!FFileHelper::SaveStringToFile(InText.ToString(), *Filename, FFileHelper::EEncodingOptions::ForceUTF8WithoutBOM))
+	{
+		UE_LOG(LogSourceControl, Error, TEXT("Failed to write to temp file: %s"), *Filename);
+	}
+}
+
+FScopedTempFile::~FScopedTempFile()
+{
+	if (FPaths::FileExists(Filename))
+	{
+		if (!FPlatformFileManager::Get().GetPlatformFile().DeleteFile(*Filename))
+		{
+			UE_LOG(LogSourceControl, Error, TEXT("Failed to delete temp file: %s"), *Filename);
+		}
+	}
+}
+
+const FString& FScopedTempFile::GetFilename() const
+{
+	return Filename;
+}
diff --git a/Engine/Plugins/Developer/PlasticSourceControl/Source/PlasticSourceControl/Private/ScopedTempFile.h b/Engine/Plugins/Developer/PlasticSourceControl/Source/PlasticSourceControl/Private/ScopedTempFile.h
new file mode 100644
--- /dev/null
+++ b/Engine/Plugins/Developer/PlasticSourceControl/Source/PlasticSourceControl/Private/ScopedTempFile.h
@@ -0,0 +1,29 @@
+// Copyright Epic Games, Inc. All Rights Reserved.
+
+#pragma once
+
+#include "CoreMinimal.h"
+
+/**
+ * Helper for temporary files to pass as arguments to some commands (typically for checkin multi-line text message)
+ */
+class FScopedTempFile
+{
+public:
+	/** Default constructor - only hold a temp filename */
+	FScopedTempFile();
+
+	/** Constructor - open & write string to temp file */
+	explicit FScopedTempFile(const FString& InText);
+	explicit FScopedTempFile(const FText& InText);
+
+	/** Destructor - delete temp file */
+	~FScopedTempFile();
+
+	/** Get the filename of this temp file - empty if it failed to be created */
+	const FString& GetFilename() const;
+
+private:
+	/** The filename we are writing to */
+	FString Filename;
+};
diff --git a/Engine/Plugins/Developer/PlasticSourceControl/Source/PlasticSourceControl/Private/SoftwareVersion.cpp b/Engine/Plugins/Developer/PlasticSourceControl/Source/PlasticSourceControl/Private/SoftwareVersion.cpp
--- a/Engine/Plugins/Developer/PlasticSourceControl/Source/PlasticSourceControl/Private/SoftwareVersion.cpp
+++ b/Engine/Plugins/Developer/PlasticSourceControl/Source/PlasticSourceControl/Private/SoftwareVersion.cpp
@@ -2,9 +2,9 @@
 
 #include "SoftwareVersion.h"
 
-FSoftwareVersion::FSoftwareVersion(const FString& InVersionString)
+FSoftwareVersion::FSoftwareVersion(FString&& InVersionString)
 {
-	String = InVersionString;
+	String = MoveTemp(InVersionString);
 	TArray<FString> Parts;
 	const int32 N = String.ParseIntoArray(Parts, TEXT("."));
 	if (N == 4)
@@ -16,15 +16,6 @@ FSoftwareVersion::FSoftwareVersion(const FString& InVersionString)
 	}
 }
 
-FSoftwareVersion::FSoftwareVersion(const int& InMajor, const int& InMinor, const int& InPatch, const int& InChangeset)
-{
-	String = FString::Printf(TEXT("%d.%d.%d.%d"), InMajor, InMinor, InPatch);
-	Major = InMajor;
-	Minor = InMinor;
-	Patch = InPatch;
-	Changeset = InChangeset;
-}
-
 bool operator==(const FSoftwareVersion& Rhs, const FSoftwareVersion& Lhs)
 {
 	return (Rhs.Major == Lhs.Major) && (Rhs.Minor == Lhs.Minor) && (Rhs.Patch == Lhs.Patch) && (Rhs.Changeset == Lhs.Changeset);
@@ -42,3 +33,8 @@ bool operator<(const FSoftwareVersion& Rhs, const FSoftwareVersion& Lhs)
 	if (Rhs.Changeset > Lhs.Changeset) return false;
 	return false; // Equal
 }
+
+bool operator>=(const FSoftwareVersion& Rhs, const FSoftwareVersion& Lhs)
+{
+	return !(Rhs < Lhs);
+}
diff --git a/Engine/Plugins/Developer/PlasticSourceControl/Source/PlasticSourceControl/Private/SoftwareVersion.h b/Engine/Plugins/Developer/PlasticSourceControl/Source/PlasticSourceControl/Private/SoftwareVersion.h
--- a/Engine/Plugins/Developer/PlasticSourceControl/Source/PlasticSourceControl/Private/SoftwareVersion.h
+++ b/Engine/Plugins/Developer/PlasticSourceControl/Source/PlasticSourceControl/Private/SoftwareVersion.h
@@ -9,12 +9,11 @@
 */
 struct FSoftwareVersion
 {
-	FSoftwareVersion() {}
+	FSoftwareVersion() : String(TEXT("<unknown-version>")) {}
 
-	explicit FSoftwareVersion(const FString& InVersionString);
-	explicit FSoftwareVersion(const int& InMajor, const int& InMinor, const int& InPatch, const int& InChangeset);
+	explicit FSoftwareVersion(FString&& InVersionString);
 
-	FString String = TEXT("<unknown-version>");
+	FString String;
 
 	int Major = 0;
 	int Minor = 0;
@@ -24,3 +23,4 @@ struct FSoftwareVersion
 
 bool operator==(const FSoftwareVersion& Rhs, const FSoftwareVersion& Lhs);
 bool operator<(const FSoftwareVersion& Rhs, const FSoftwareVersion& Lhs);
+bool operator>=(const FSoftwareVersion& Rhs, const FSoftwareVersion& Lhs);
