Fix a crash in the PlasticSCM module that can occur if the engine shutdown during PreInit.

#rnx

- The crash itself isn't that interesting but it does distract the user from the actual error/problem which caused the engine to shutdown during PreInit wasting time while they look into the crash instead of the original error.
- When attempting to unregister the menu extension we now call ::TryGet rather than ::Try as this will return nullptr if UToolsMenu is not loaded and will not attempt to load it.

diff --git a/Engine/Plugins/Developer/PlasticSourceControl/Source/PlasticSourceControl/Private/PlasticSourceControlMenu.cpp b/Engine/Plugins/Developer/PlasticSourceControl/Source/PlasticSourceControl/Private/PlasticSourceControlMenu.cpp
--- a/Engine/Plugins/Developer/PlasticSourceControl/Source/PlasticSourceControl/Private/PlasticSourceControlMenu.cpp
+++ b/Engine/Plugins/Developer/PlasticSourceControl/Source/PlasticSourceControl/Private/PlasticSourceControlMenu.cpp
@@ -32,6 +32,11 @@
 
 void FPlasticSourceControlMenu::Register()
 {
+	if (bHasRegistered)
+	{
+		return;
+	}
+
 	// Register the menu extension with the level editor
 	FToolMenuOwnerScoped SourceControlMenuOwner("PlasticSourceControlMenu");
 	if (UToolMenus* ToolMenus = UToolMenus::Get())
@@ -40,15 +45,23 @@ void FPlasticSourceControlMenu::Register()
 		FToolMenuSection& Section = SourceControlMenu->AddSection("PlasticSourceControlActions", LOCTEXT("PlasticSourceControlMenuHeadingActions", "Plastic SCM"), FToolMenuInsert(NAME_None, EToolMenuInsertType::First));
 
 		AddMenuExtension(Section);
+
+		bHasRegistered = true;
 	}
 }
 
 void FPlasticSourceControlMenu::Unregister()
 {
+	if (!bHasRegistered)
+	{
+		return;
+	}
+
 	// Unregister the menu extension from the level editor
-	if (UToolMenus* ToolMenus = UToolMenus::Get())
+	if (UToolMenus* ToolMenus = UToolMenus::TryGet())
 	{
 		UToolMenus::Get()->UnregisterOwnerByName("PlasticSourceControlMenu");
+		bHasRegistered = false;
 	}
 }
 
diff --git a/Engine/Plugins/Developer/PlasticSourceControl/Source/PlasticSourceControl/Private/PlasticSourceControlMenu.h b/Engine/Plugins/Developer/PlasticSourceControl/Source/PlasticSourceControl/Private/PlasticSourceControlMenu.h
--- a/Engine/Plugins/Developer/PlasticSourceControl/Source/PlasticSourceControl/Private/PlasticSourceControlMenu.h
+++ b/Engine/Plugins/Developer/PlasticSourceControl/Source/PlasticSourceControl/Private/PlasticSourceControlMenu.h
@@ -39,6 +39,10 @@ private:
 	void DisplayFailureNotification(const FName& InOperationName);
 
 private:
+
+	/** Tracks if the menu extension has been registered with the editor or not */
+	bool bHasRegistered = false;
+
 	/** Current source control operation from extended menu if any */
 	TWeakPtr<class SNotificationItem> OperationInProgressNotification;
 
