[Plastic SCM] Fix implementation of ISourceControlProvider::UsesFileRevisions()

PR #10271

diff --git a/Engine/Plugins/Developer/PlasticSourceControl/Source/PlasticSourceControl/Private/PlasticSourceControlProvider.cpp b/Engine/Plugins/Developer/PlasticSourceControl/Source/PlasticSourceControl/Private/PlasticSourceControlProvider.cpp
--- a/Engine/Plugins/Developer/PlasticSourceControl/Source/PlasticSourceControl/Private/PlasticSourceControlProvider.cpp
+++ b/Engine/Plugins/Developer/PlasticSourceControl/Source/PlasticSourceControl/Private/PlasticSourceControlProvider.cpp
@@ -432,7 +432,9 @@ bool FPlasticSourceControlProvider::UsesCheckout() const
 
 bool FPlasticSourceControlProvider::UsesFileRevisions() const
 {
-	return true;
+	// Only a Partial/Gluon workspace can sync/update files individually, operating on revisions (can use the context menu)
+	// while a regular workspace can only update all the files as a whole, operating at the changeset level (requires the global menu)
+	return IsPartialWorkspace();
 }
 
 bool FPlasticSourceControlProvider::UsesSnapshots() const
