Fixed mismatched copy constuctor/assignment operator in FPlasticSourceControlChangelist.

diff --git a/Engine/Plugins/Developer/PlasticSourceControl/Source/PlasticSourceControl/Private/PlasticSourceControlChangelist.h b/Engine/Plugins/Developer/PlasticSourceControl/Source/PlasticSourceControl/Private/PlasticSourceControlChangelist.h
--- a/Engine/Plugins/Developer/PlasticSourceControl/Source/PlasticSourceControl/Private/PlasticSourceControlChangelist.h
+++ b/Engine/Plugins/Developer/PlasticSourceControl/Source/PlasticSourceControl/Private/PlasticSourceControlChangelist.h
@@ -20,8 +20,6 @@ public:
 	{
 	}
 
-	FPlasticSourceControlChangelist(const FPlasticSourceControlChangelist& InOther) = default;
-
 	virtual bool CanDelete() const override
 	{
 		return !IsDefault();
