Truncation warning fixes for PlasticSourceControl module

#rnx

diff --git a/Engine/Plugins/Developer/PlasticSourceControl/Source/PlasticSourceControl/PlasticSourceControl.Build.cs b/Engine/Plugins/Developer/PlasticSourceControl/Source/PlasticSourceControl/PlasticSourceControl.Build.cs
--- a/Engine/Plugins/Developer/PlasticSourceControl/Source/PlasticSourceControl/PlasticSourceControl.Build.cs
+++ b/Engine/Plugins/Developer/PlasticSourceControl/Source/PlasticSourceControl/PlasticSourceControl.Build.cs
@@ -24,5 +24,7 @@ public class PlasticSourceControl : ModuleRules
 				"ToolMenus",
 			}
 		);
+
+		UnsafeTypeCastWarningLevel = WarningLevel.Error;
 	}
 }
diff --git a/Engine/Plugins/Developer/PlasticSourceControl/Source/PlasticSourceControl/Private/SPlasticSourceControlSettings.cpp b/Engine/Plugins/Developer/PlasticSourceControl/Source/PlasticSourceControl/Private/SPlasticSourceControlSettings.cpp
--- a/Engine/Plugins/Developer/PlasticSourceControl/Source/PlasticSourceControl/Private/SPlasticSourceControlSettings.cpp
+++ b/Engine/Plugins/Developer/PlasticSourceControl/Source/PlasticSourceControl/Private/SPlasticSourceControlSettings.cpp
@@ -349,7 +349,7 @@ void SPlasticSourceControlSettings::Construct(const FArguments& InArgs)
 				.ToolTipText(LOCTEXT("PlasticInitWorkspace_Tooltip", "Create and initialize a new Plastic SCM workspace and repository for the current project"))
 				.OnClicked(this, &SPlasticSourceControlSettings::OnClickedInitializePlasticWorkspace)
 				.HAlign(HAlign_Center)
-				.ContentPadding(6)
+				.ContentPadding(6.0f)
 			]
 		]
 		// Button to add a 'ignore.conf' file on an existing Workspace
