PR #11295: [Source Control] Update Plastic SCM provider to 1.9.0 (from 1.8.0)

Compared to the PR:
- Fixed non-Unity compiles by adding several header includes.
- Fixed Linux compile error in SPlasticSourceControlBranchesWidget.cpp: FPlasticSourceControlBranchPtr vs FPlasticSourceControlBranchRef usage in ::GetSelectedBranches().

diff --git a/Engine/Plugins/Developer/PlasticSourceControl/PlasticSourceControl.uplugin b/Engine/Plugins/Developer/PlasticSourceControl/PlasticSourceControl.uplugin
--- a/Engine/Plugins/Developer/PlasticSourceControl/PlasticSourceControl.uplugin
+++ b/Engine/Plugins/Developer/PlasticSourceControl/PlasticSourceControl.uplugin
@@ -1,15 +1,15 @@
 {
 	"FileVersion": 3,
-	"Version": 80,
-	"VersionName": "1.8.0",
+	"Version": 90,
+	"VersionName": "1.9.0",
 	"FriendlyName": "Plastic SCM",
-	"Description": "Plastic source control management",
+	"Description": "Unity Version Control (formerly Plastic SCM)",
 	"Category": "Source Control",
 	"CreatedBy": "SRombauts for Codice Software.",
 	"CreatedByURL": "https://srombauts.github.io",
-	"DocsURL": "[REDACTED_PROFILE_URL]/UEPlasticPlugin",
+	"DocsURL": "https://docs.unity.com/ugs/en-us/manual/devops/manual/vcs-plugins/unreal-plugin",
 	"MarketplaceURL": "",
-	"SupportURL": "https://plasticscm.com/support",
+	"SupportURL": "https://support.unity.com/hc/en-us/requests/new?ticket_form_id=360001051792",
 	"EnabledByDefault": true,
 	"CanContainContent": false,
 	"IsBetaVersion": false,
diff --git a/Engine/Plugins/Developer/PlasticSourceControl/Source/PlasticSourceControl/PlasticSourceControl.Build.cs b/Engine/Plugins/Developer/PlasticSourceControl/Source/PlasticSourceControl/PlasticSourceControl.Build.cs
--- a/Engine/Plugins/Developer/PlasticSourceControl/Source/PlasticSourceControl/PlasticSourceControl.Build.cs
+++ b/Engine/Plugins/Developer/PlasticSourceControl/Source/PlasticSourceControl/PlasticSourceControl.Build.cs
@@ -22,6 +22,7 @@ public class PlasticSourceControl : ModuleRules
 				"AssetRegistry",
 				"DeveloperSettings",
 				"ToolMenus",
+				"ContentBrowser",
 			}
 		);
 
diff --git a/Engine/Plugins/Developer/PlasticSourceControl/Source/PlasticSourceControl/Private/PackageUtils.cpp b/Engine/Plugins/Developer/PlasticSourceControl/Source/PlasticSourceControl/Private/PackageUtils.cpp
--- a/Engine/Plugins/Developer/PlasticSourceControl/Source/PlasticSourceControl/Private/PackageUtils.cpp
+++ b/Engine/Plugins/Developer/PlasticSourceControl/Source/PlasticSourceControl/Private/PackageUtils.cpp
@@ -3,11 +3,11 @@
 #include "PackageUtils.h"
 
 #include "Async/Async.h"
-#include "PackageTools.h"
 #include "Editor.h"
-#include "UObject/Linker.h"
-
+#include "FileHelpers.h"
 #include "ISourceControlModule.h" // LogSourceControl
+#include "PackageTools.h"
+#include "UObject/Linker.h"
 
 namespace PackageUtils
 {
@@ -25,6 +25,56 @@ static UWorld* GetCurrentWorld()
 	return nullptr;
 }
 
+TArray<FString> AssetDataToFileNames(const TArray<FAssetData>& InAssetObjectPaths)
+{
+	TArray<FString> FileNames;
+	FileNames.Reserve(InAssetObjectPaths.Num());
+	for (const FAssetData& AssetObjectPath : InAssetObjectPaths)
+	{
+		//  Tries to find the package in memory if it is loaded, otherwise loads it
+		if (UPackage* Package = AssetObjectPath.GetPackage())
+		{
+			const FString PackageExtension = Package->ContainsMap() ? FPackageName::GetMapPackageExtension() : FPackageName::GetAssetPackageExtension();
+			FString PackageFilename = FPackageName::LongPackageNameToFilename(Package->GetName(), PackageExtension);
+			FileNames.Add(MoveTemp(PackageFilename));
+		}
+	}
+	return FileNames;
+}
+
+/// Prompt to save or discard all packages
+bool SaveDirtyPackages()
+{
+	const bool bPromptUserToSave = true;
+	const bool bSaveMapPackages = true;
+	const bool bSaveContentPackages = true;
+	const bool bFastSave = false;
+	const bool bNotifyNoPackagesSaved = false;
+	const bool bCanBeDeclined = true; // If the user clicks "don't save" this will continue and lose their changes
+	bool bHadPackagesToSave = false;
+
+	bool bSaved = FEditorFileUtils::SaveDirtyPackages(bPromptUserToSave, bSaveMapPackages, bSaveContentPackages, bFastSave, bNotifyNoPackagesSaved, bCanBeDeclined, &bHadPackagesToSave);
+
+	// bSaved can be true if the user selects to not save an asset by un-checking it and clicking "save"
+	if (bSaved)
+	{
+		TArray<UPackage*> DirtyPackages;
+		FEditorFileUtils::GetDirtyWorldPackages(DirtyPackages);
+		FEditorFileUtils::GetDirtyContentPackages(DirtyPackages);
+		bSaved = DirtyPackages.Num() == 0;
+	}
+
+	return bSaved;
+}
+
+/// Find all packages in Content directory
+TArray<FString> ListAllPackages()
+{
+	TArray<FString> PackageFilePaths;
+	FPackageName::FindPackagesInDirectory(PackageFilePaths, FPaths::ConvertRelativePathToFull(FPaths::ProjectContentDir()));
+	return PackageFilePaths;
+}
+
 // Find the packages corresponding to the files, if they are loaded in memory (won't load them)
 // Note: Extracted from AssetViewUtils::SyncPathsFromSourceControl()
 static TArray<UPackage*> FileNamesToLoadedPackages(const TArray<FString>& InFiles)
@@ -133,7 +183,7 @@ static TArray<UPackage*> ListPackagesToReload(const TArray<FString>& InFiles)
 // Note: Extracted from AssetViewUtils::SyncPathsFromSourceControl()
 void ReloadPackages(TArray<UPackage*>& InPackages)
 {
-	UE_LOG(LogSourceControl, Log, TEXT("Reloading %d Packages..."), InPackages.Num());
+	UE_LOG(LogSourceControl, Verbose, TEXT("Reloading %d Packages..."), InPackages.Num());
 
 	// Syncing may have deleted some packages, so we need to unload those rather than re-load them...
 	// Note: we will store the package using weak pointers here otherwise we might have garbage collection issues after the ReloadPackages call
@@ -174,6 +224,10 @@ void ReloadPackages(const TArray<FString>& InFiles)
 	{
 		ReloadPackages(PackagesToReload);
 	}
+	else
+	{
+		UE_LOG(LogSourceControl, Verbose, TEXT("No package to reload"));
+	}
 }
 
 void ReloadPackagesInMainThread(const TArray<FString>& InFiles)
diff --git a/Engine/Plugins/Developer/PlasticSourceControl/Source/PlasticSourceControl/Private/PackageUtils.h b/Engine/Plugins/Developer/PlasticSourceControl/Source/PlasticSourceControl/Private/PackageUtils.h
--- a/Engine/Plugins/Developer/PlasticSourceControl/Source/PlasticSourceControl/Private/PackageUtils.h
+++ b/Engine/Plugins/Developer/PlasticSourceControl/Source/PlasticSourceControl/Private/PackageUtils.h
@@ -3,9 +3,15 @@
 #pragma once
 
 #include "CoreMinimal.h"
+#include "AssetRegistry/AssetData.h"
 
 namespace PackageUtils
 {
+	TArray<FString> AssetDataToFileNames(const TArray<FAssetData>& InAssetObjectPaths);
+
+	bool SaveDirtyPackages();
+	TArray<FString> ListAllPackages();
+
 	void UnlinkPackages(const TArray<FString>& InFiles);
 	void UnlinkPackagesInMainThread(const TArray<FString>& InFiles);
 
diff --git a/Engine/Plugins/Developer/PlasticSourceControl/Source/PlasticSourceControl/Private/PlasticSourceControlBranch.h b/Engine/Plugins/Developer/PlasticSourceControl/Source/PlasticSourceControl/Private/PlasticSourceControlBranch.h
new file mode 100644
--- /dev/null
+++ b/Engine/Plugins/Developer/PlasticSourceControl/Source/PlasticSourceControl/Private/PlasticSourceControlBranch.h
@@ -0,0 +1,33 @@
+// Copyright Epic Games, Inc. All Rights Reserved.
+
+#pragma once
+
+#include "CoreMinimal.h"
+
+class FPlasticSourceControlBranch
+{
+public:
+	FPlasticSourceControlBranch() = default;
+	FPlasticSourceControlBranch(const FString& InName, const FString& InRepository, const FString& InCreatedBy, const FDateTime& InDate, const FString& InComment)
+		: Name(InName)
+		, Repository(InRepository)
+		, CreatedBy(InCreatedBy)
+		, Date(InDate)
+		, Comment(InComment)
+	{}
+	FString Name;
+	FString Repository;
+	FString CreatedBy;
+	FDateTime Date;
+	FString Comment;
+
+	void PopulateSearchString(TArray<FString>& OutStrings) const
+	{
+		OutStrings.Emplace(Name);
+		OutStrings.Emplace(CreatedBy);
+		OutStrings.Emplace(Comment);
+	}
+};
+
+typedef TSharedRef<class FPlasticSourceControlBranch, ESPMode::ThreadSafe> FPlasticSourceControlBranchRef;
+typedef TSharedPtr<class FPlasticSourceControlBranch, ESPMode::ThreadSafe> FPlasticSourceControlBranchPtr;
\ No newline at end of file
diff --git a/Engine/Plugins/Developer/PlasticSourceControl/Source/PlasticSourceControl/Private/PlasticSourceControlBranchesWindow.cpp b/Engine/Plugins/Developer/PlasticSourceControl/Source/PlasticSourceControl/Private/PlasticSourceControlBranchesWindow.cpp
new file mode 100644
--- /dev/null
+++ b/Engine/Plugins/Developer/PlasticSourceControl/Source/PlasticSourceControl/Private/PlasticSourceControlBranchesWindow.cpp
@@ -0,0 +1,51 @@
+// Copyright Epic Games, Inc. All Rights Reserved.
+
+#include "PlasticSourceControlBranchesWindow.h"
+
+#include "Widgets/Docking/SDockTab.h"
+
+#include "PlasticSourceControlStyle.h"
+#include "SPlasticSourceControlBranchesWidget.h"
+
+#define LOCTEXT_NAMESPACE "PlasticSourceControlWindow"
+
+static const FName PlasticSourceControlWindowTabName("PlasticSourceControlWindow");
+
+void FPlasticSourceControlBranchesWindow::Register()
+{
+	FPlasticSourceControlStyle::Initialize();
+	FPlasticSourceControlStyle::ReloadTextures();
+
+	FGlobalTabmanager::Get()->RegisterNomadTabSpawner(PlasticSourceControlWindowTabName, FOnSpawnTab::CreateRaw(this, &FPlasticSourceControlBranchesWindow::OnSpawnTab))
+		.SetDisplayName(LOCTEXT("PlasticSourceControlWindowTabTitle", "View Branches"))
+		.SetMenuType(ETabSpawnerMenuType::Hidden)
+		.SetIcon(FSlateIcon(FPlasticSourceControlStyle::Get().GetStyleSetName(), "PlasticSourceControl.PluginIcon.Small"));
+}
+
+void FPlasticSourceControlBranchesWindow::Unregister()
+{
+	FGlobalTabmanager::Get()->UnregisterNomadTabSpawner(PlasticSourceControlWindowTabName);
+
+	FPlasticSourceControlStyle::Shutdown();
+}
+
+TSharedRef<SDockTab> FPlasticSourceControlBranchesWindow::OnSpawnTab(const FSpawnTabArgs& SpawnTabArgs)
+{
+	return SNew(SDockTab)
+		.TabRole(ETabRole::NomadTab)
+		[
+			CreateBranchesWidget().ToSharedRef()
+		];
+}
+
+void FPlasticSourceControlBranchesWindow::OpenTab()
+{
+	FGlobalTabmanager::Get()->TryInvokeTab(PlasticSourceControlWindowTabName);
+}
+
+TSharedPtr<SWidget> FPlasticSourceControlBranchesWindow::CreateBranchesWidget()
+{
+	return SNew(SPlasticSourceControlBranchesWidget);
+}
+
+#undef LOCTEXT_NAMESPACE
\ No newline at end of file
diff --git a/Engine/Plugins/Developer/PlasticSourceControl/Source/PlasticSourceControl/Private/PlasticSourceControlChangelist.h b/Engine/Plugins/Developer/PlasticSourceControl/Source/PlasticSourceControl/Private/PlasticSourceControlChangelist.h
--- a/Engine/Plugins/Developer/PlasticSourceControl/Source/PlasticSourceControl/Private/PlasticSourceControlChangelist.h
+++ b/Engine/Plugins/Developer/PlasticSourceControl/Source/PlasticSourceControl/Private/PlasticSourceControlChangelist.h
@@ -7,7 +7,7 @@
 #include "ISourceControlChangelist.h"
 
 /**
- * Unique Identifier of a changelist under source control: a "name" in Plastic SCM
+ * Unique Identifier of a changelist under source control: a "name" in Unity Version Control
  */
 class FPlasticSourceControlChangelist : public ISourceControlChangelist
 {
diff --git a/Engine/Plugins/Developer/PlasticSourceControl/Source/PlasticSourceControl/Private/PlasticSourceControlCommand.cpp b/Engine/Plugins/Developer/PlasticSourceControl/Source/PlasticSourceControl/Private/PlasticSourceControlCommand.cpp
--- a/Engine/Plugins/Developer/PlasticSourceControl/Source/PlasticSourceControl/Private/PlasticSourceControlCommand.cpp
+++ b/Engine/Plugins/Developer/PlasticSourceControl/Source/PlasticSourceControl/Private/PlasticSourceControlCommand.cpp
@@ -1,12 +1,12 @@
 // Copyright Epic Games, Inc. All Rights Reserved.
 
 #include "PlasticSourceControlCommand.h"
+
 #include "PlasticSourceControlModule.h"
+
 #include "ISourceControlOperation.h"
-#include "Modules/ModuleManager.h"
 #include "HAL/PlatformTime.h"
 
-
 FPlasticSourceControlCommand::FPlasticSourceControlCommand(const FSourceControlOperationRef& InOperation, const FPlasticSourceControlWorkerRef& InWorker, const FSourceControlOperationComplete& InOperationCompleteDelegate)
 	: Operation(InOperation)
 	, Worker(InWorker)
@@ -43,3 +43,22 @@ void FPlasticSourceControlCommand::DoThreadedWork()
 	Concurrency = EConcurrency::Asynchronous;
 	DoWork();
 }
+
+ECommandResult::Type FPlasticSourceControlCommand::ReturnResults()
+{
+	// Save any messages that have accumulated
+	for (FString& String : InfoMessages)
+	{
+		Operation->AddInfoMessge(FText::FromString(String));
+	}
+	for (FString& String : ErrorMessages)
+	{
+		Operation->AddErrorMessge(FText::FromString(String));
+	}
+
+	// run the completion delegate if we have one bound
+	ECommandResult::Type Result = bCommandSuccessful ? ECommandResult::Succeeded : ECommandResult::Failed;
+	OperationCompleteDelegate.ExecuteIfBound(Operation, Result);
+
+	return Result;
+}
diff --git a/Engine/Plugins/Developer/PlasticSourceControl/Source/PlasticSourceControl/Private/PlasticSourceControlCommand.h b/Engine/Plugins/Developer/PlasticSourceControl/Source/PlasticSourceControl/Private/PlasticSourceControlCommand.h
--- a/Engine/Plugins/Developer/PlasticSourceControl/Source/PlasticSourceControl/Private/PlasticSourceControlCommand.h
+++ b/Engine/Plugins/Developer/PlasticSourceControl/Source/PlasticSourceControl/Private/PlasticSourceControlCommand.h
@@ -37,6 +37,9 @@ public:
 	 */ 
 	virtual void DoThreadedWork() override;
 
+	/** Save any results and call any registered callbacks. */
+	ECommandResult::Type ReturnResults();
+
 public:
 	/** Path to the root of the Plastic workspace: can be the GameDir itself, or any parent directory (found by the "Connect" operation) */
 	FString PathToWorkspaceRoot;
@@ -94,4 +97,4 @@ public:
 
 	/**Potential error message storage*/
 	TArray<FString> ErrorMessages;
-};
\ No newline at end of file
+};
diff --git a/Engine/Plugins/Developer/PlasticSourceControl/Source/PlasticSourceControl/Private/PlasticSourceControlConsole.cpp b/Engine/Plugins/Developer/PlasticSourceControl/Source/PlasticSourceControl/Private/PlasticSourceControlConsole.cpp
--- a/Engine/Plugins/Developer/PlasticSourceControl/Source/PlasticSourceControl/Private/PlasticSourceControlConsole.cpp
+++ b/Engine/Plugins/Developer/PlasticSourceControl/Source/PlasticSourceControl/Private/PlasticSourceControlConsole.cpp
@@ -2,17 +2,17 @@
 
 #include "PlasticSourceControlConsole.h"
 
-#include "ISourceControlModule.h"
-
 #include "PlasticSourceControlUtils.h"
 
+#include "ISourceControlModule.h"
+
 void FPlasticSourceControlConsole::Register()
 {
 	if (!CmConsoleCommand.IsValid())
 	{
 		CmConsoleCommand = MakeUnique<FAutoConsoleCommand>(
 			TEXT("cm"),
-			TEXT("PlasticSCM Command Line Interface.\n")
+			TEXT("Unity Version Control (formerly Plastic SCM) Command Line Interface (cm).\n")
 			TEXT("Run any 'cm' command directly from the Unreal Editor Console.\n")
 			TEXT("Type 'cm showcommands' to get a command list."),
 			FConsoleCommandWithArgsDelegate::CreateRaw(this, &FPlasticSourceControlConsole::ExecutePlasticConsoleCommand));
diff --git a/Engine/Plugins/Developer/PlasticSourceControl/Source/PlasticSourceControl/Private/PlasticSourceControlConsole.h b/Engine/Plugins/Developer/PlasticSourceControl/Source/PlasticSourceControl/Private/PlasticSourceControlConsole.h
--- a/Engine/Plugins/Developer/PlasticSourceControl/Source/PlasticSourceControl/Private/PlasticSourceControlConsole.h
+++ b/Engine/Plugins/Developer/PlasticSourceControl/Source/PlasticSourceControl/Private/PlasticSourceControlConsole.h
@@ -18,7 +18,7 @@ public:
 	void Unregister();
 
 private:
-	// PlasticSCM Command Line Interface: Run 'cm' commands directly from the Unreal Editor Console.
+	// Unity Version Control Command Line Interface: Run 'cm' commands directly from the Unreal Editor Console.
 	void ExecutePlasticConsoleCommand(const TArray<FString>& a_args);
 
 	/** Console command for interacting with 'cm' CLI directly */
diff --git a/Engine/Plugins/Developer/PlasticSourceControl/Source/PlasticSourceControl/Private/PlasticSourceControlMenu.cpp b/Engine/Plugins/Developer/PlasticSourceControl/Source/PlasticSourceControl/Private/PlasticSourceControlMenu.cpp
--- a/Engine/Plugins/Developer/PlasticSourceControl/Source/PlasticSourceControl/Private/PlasticSourceControlMenu.cpp
+++ b/Engine/Plugins/Developer/PlasticSourceControl/Source/PlasticSourceControl/Private/PlasticSourceControlMenu.cpp
@@ -2,34 +2,37 @@
 
 #include "PlasticSourceControlMenu.h"
 
+#include "PlasticSourceControlBranchesWindow.h"
 #include "PlasticSourceControlModule.h"
 #include "PlasticSourceControlProvider.h"
 #include "PlasticSourceControlOperations.h"
+#include "SPlasticSourceControlStatusBar.h"
 
 #include "ISourceControlModule.h"
 #include "ISourceControlOperation.h"
 #include "SourceControlOperations.h"
 
+#include "ContentBrowserMenuContexts.h"
 #include "Interfaces/IPluginManager.h"
+#include "Modules/ModuleManager.h"
 #include "LevelEditor.h"
-#include "UObject/Linker.h"
-#include "Widgets/Notifications/SNotificationList.h"
-#include "Framework/Notifications/NotificationManager.h"
 #include "Misc/MessageDialog.h"
 #include "Styling/AppStyle.h"
 
 #include "PackageUtils.h"
-#include "FileHelpers.h"
 #include "ISettingsModule.h"
 
 #include "Logging/MessageLog.h"
 
 #include "ToolMenus.h"
-#include "ToolMenuContext.h"
 #include "ToolMenuMisc.h"
 
 #define LOCTEXT_NAMESPACE "PlasticSourceControl"
 
+FName FPlasticSourceControlMenu::UnityVersionControlMainMenuOwnerName = TEXT("UnityVersionControlMenu");
+FName FPlasticSourceControlMenu::UnityVersionControlAssetContextLocksMenuOwnerName = TEXT("UnityVersionControlContextLocksMenu");
+FName FPlasticSourceControlMenu::UnityVersionControlStatusBarMenuOwnerName = TEXT("UnityVersionControlStatusBarMenu");
+
 void FPlasticSourceControlMenu::Register()
 {
 	if (bHasRegistered)
@@ -37,17 +40,11 @@ void FPlasticSourceControlMenu::Register()
 		return;
 	}
 
-	// Register the menu extension with the level editor
-	FToolMenuOwnerScoped SourceControlMenuOwner("PlasticSourceControlMenu");
-	if (UToolMenus* ToolMenus = UToolMenus::Get())
-	{
-		UToolMenu* SourceControlMenu = ToolMenus->ExtendMenu("StatusBar.ToolBar.SourceControl");
-		FToolMenuSection& Section = SourceControlMenu->AddSection("PlasticSourceControlActions", LOCTEXT("PlasticSourceControlMenuHeadingActions", "Plastic SCM"), FToolMenuInsert(NAME_None, EToolMenuInsertType::First));
-
-		AddMenuExtension(Section);
+	// Register the menu extensions with the level editor
+	ExtendRevisionControlMenu();
+	ExtendAssetContextMenu();
 
-		bHasRegistered = true;
-	}
+	ExtendToolbarWithStatusBarWidget();
 }
 
 void FPlasticSourceControlMenu::Unregister()
@@ -57,83 +54,227 @@ void FPlasticSourceControlMenu::Unregister()
 		return;
 	}
 
-	// Unregister the menu extension from the level editor
 	if (UToolMenus* ToolMenus = UToolMenus::TryGet())
 	{
-		UToolMenus::Get()->UnregisterOwnerByName("PlasticSourceControlMenu");
+		ToolMenus->UnregisterOwnerByName(UnityVersionControlMainMenuOwnerName);
+		ToolMenus->UnregisterOwnerByName(UnityVersionControlAssetContextLocksMenuOwnerName);
+		ToolMenus->UnregisterOwnerByName(UnityVersionControlStatusBarMenuOwnerName);
 		bHasRegistered = false;
 	}
 }
 
-bool FPlasticSourceControlMenu::IsSourceControlConnected() const
+void FPlasticSourceControlMenu::ExtendToolbarWithStatusBarWidget()
 {
-	const ISourceControlProvider& Provider = ISourceControlModule::Get().GetProvider();
-	return Provider.IsEnabled() && Provider.IsAvailable();
+	const FToolMenuOwnerScoped SourceControlMenuOwner(UnityVersionControlStatusBarMenuOwnerName);
+
+	UToolMenu* ToolbarMenu = UToolMenus::Get()->ExtendMenu("LevelEditor.StatusBar.ToolBar");
+	FToolMenuSection& Section = ToolbarMenu->AddSection("Unity Version Control", FText::GetEmpty(), FToolMenuInsert("SourceControl", EToolMenuInsertType::Before));
+
+	Section.AddEntry(
+		FToolMenuEntry::InitWidget("UnityVersionControlStatusBar", SNew(SPlasticSourceControlStatusBar), FText::GetEmpty(), true, false)
+	);
 }
 
-/// Prompt to save or discard all packages
-bool FPlasticSourceControlMenu::SaveDirtyPackages()
+void FPlasticSourceControlMenu::ExtendRevisionControlMenu()
 {
-	const bool bPromptUserToSave = true;
-	const bool bSaveMapPackages = true;
-	const bool bSaveContentPackages = true;
-	const bool bFastSave = false;
-	const bool bNotifyNoPackagesSaved = false;
-	const bool bCanBeDeclined = true; // If the user clicks "don't save" this will continue and lose their changes
-	bool bHadPackagesToSave = false;
-
-	bool bSaved = FEditorFileUtils::SaveDirtyPackages(bPromptUserToSave, bSaveMapPackages, bSaveContentPackages, bFastSave, bNotifyNoPackagesSaved, bCanBeDeclined, &bHadPackagesToSave);
-
-	// bSaved can be true if the user selects to not save an asset by un-checking it and clicking "save"
-	if (bSaved)
+	const FToolMenuOwnerScoped SourceControlMenuOwner(UnityVersionControlMainMenuOwnerName);
+
+	if (UToolMenu* SourceControlMenu = UToolMenus::Get()->ExtendMenu("StatusBar.ToolBar.SourceControl"))
 	{
-		TArray<UPackage*> DirtyPackages;
-		FEditorFileUtils::GetDirtyWorldPackages(DirtyPackages);
-		FEditorFileUtils::GetDirtyContentPackages(DirtyPackages);
-		bSaved = DirtyPackages.Num() == 0;
+		FToolMenuSection& Section = SourceControlMenu->AddSection("PlasticSourceControlActions", LOCTEXT("PlasticSourceControlMenuHeadingActions", "Unity Version Control"), FToolMenuInsert(NAME_None, EToolMenuInsertType::First));
+
+		AddMenuExtension(Section);
+
+		bHasRegistered = true;
 	}
 
-	return bSaved;
+	if (UToolMenu* ToolsMenu = UToolMenus::Get()->ExtendMenu("MainFrame.MainMenu.Tools"))
+	{
+		if (FToolMenuSection* Section = ToolsMenu->FindSection("Source Control"))
+		{
+			AddViewBranches(*Section);
+		}
+	}
 }
 
-/// Find all packages in Content directory
-TArray<FString> FPlasticSourceControlMenu::ListAllPackages()
+void FPlasticSourceControlMenu::ExtendAssetContextMenu()
 {
-	TArray<FString> PackageFilePaths;
-	FPackageName::FindPackagesInDirectory(PackageFilePaths, FPaths::ConvertRelativePathToFull(FPaths::ProjectContentDir()));
-	return PackageFilePaths;
+	const FToolMenuOwnerScoped SourceControlMenuOwner(UnityVersionControlAssetContextLocksMenuOwnerName);
+	if (UToolMenu* const Menu = UToolMenus::Get()->ExtendMenu(TEXT("ContentBrowser.AssetContextMenu")))
+	{
+		FToolMenuSection& Section = Menu->AddSection(TEXT("PlasticAssetContextLocksMenuSection"), FText::GetEmpty(), FToolMenuInsert("AssetContextReferences", EToolMenuInsertType::After));
+		Section.AddDynamicEntry(TEXT("PlasticActions"), FNewToolMenuSectionDelegate::CreateLambda([this](FToolMenuSection& InSection)
+		{
+			UContentBrowserAssetContextMenuContext* Context = InSection.FindContext<UContentBrowserAssetContextMenuContext>();
+
+			if (!Context || !Context->bCanBeModified || Context->SelectedAssets.Num() == 0 || !ensure(FPlasticSourceControlModule::IsLoaded()))
+			{
+				return;
+			}
+			TArray<FAssetData> AssetObjectPaths = Context->SelectedAssets;
+
+			InSection.AddSubMenu(
+				TEXT("PlasticActionsSubMenu"),
+				LOCTEXT("Plastic_ContextMenu", "Revision Control Locks"),
+				FText::GetEmpty(),
+				FNewMenuDelegate::CreateRaw(this, &FPlasticSourceControlMenu::GeneratePlasticAssetContextMenu, MoveTemp(AssetObjectPaths)),
+				false,
+				FSlateIcon(FAppStyle::GetAppStyleSetName(), "PropertyWindow.Locked")
+			);
+		}));
+	}
 }
 
-void FPlasticSourceControlMenu::SyncProjectClicked()
+void FPlasticSourceControlMenu::GeneratePlasticAssetContextMenu(FMenuBuilder& MenuBuilder, TArray<FAssetData> InAssetObjectPaths)
 {
-	if (!OperationInProgressNotification.IsValid())
+	MenuBuilder.BeginSection("AssetPlasticActions", LOCTEXT("UnityVersionControlAssetContextLocksMenuHeading", "Unity Version Control Locks"));
+
+	{
+		MenuBuilder.AddMenuEntry(
+			LOCTEXT("PlasticReleaseLock", "Release Lock"),
+			LOCTEXT("PlasticReleaseLockTooltip", "Release Lock(s) on the selected assets. Requires administrator privileges on the server."),
+			FSlateIcon(FAppStyle::GetAppStyleSetName(), "PropertyWindow.Unlocked"),
+			FUIAction(
+				FExecuteAction::CreateRaw(this, &FPlasticSourceControlMenu::ExecuteReleaseLocks, InAssetObjectPaths),
+				FCanExecuteAction::CreateRaw(this, &FPlasticSourceControlMenu::CanReleaseLocks, InAssetObjectPaths)
+			)
+		);
+	}
+
 	{
-		const bool bSaved = SaveDirtyPackages();
-		if (bSaved)
+		MenuBuilder.AddMenuEntry(
+			LOCTEXT("PlasticRemoveLock", "Remove Lock"),
+			LOCTEXT("PlasticRemoveLockTooltip", "Remove/Delete Lock(s) on the selected assets. Requires administrator privileges on the server."),
+			FSlateIcon(FAppStyle::GetAppStyleSetName(), "PropertyWindow.Unlocked"),
+			FUIAction(
+				FExecuteAction::CreateRaw(this, &FPlasticSourceControlMenu::ExecuteRemoveLocks, InAssetObjectPaths),
+				FCanExecuteAction::CreateRaw(this, &FPlasticSourceControlMenu::CanRemoveLocks, InAssetObjectPaths)
+			)
+		);
+	}
+
+	FString OrganizationName = FPlasticSourceControlModule::Get().GetProvider().GetCloudOrganization();
+	if (!OrganizationName.IsEmpty())
+	{
+		MenuBuilder.AddMenuEntry(
+			LOCTEXT("PlasticLockRulesURL", "Configure Lock Rules"),
+			LOCTEXT("PlasticLockRulesURLTooltip", "Navigate to lock rules configuration page in the Unity Dashboard."),
+			FSlateIcon(FAppStyle::GetAppStyleSetName(), "PropertyWindow.Locked"),
+			FUIAction(
+				FExecuteAction::CreateRaw(this, &FPlasticSourceControlMenu::VisitLockRulesURLClicked, MoveTemp(OrganizationName)),
+				FCanExecuteAction()
+			)
+		);
+	}
+
+	MenuBuilder.EndSection();
+}
+
+bool FPlasticSourceControlMenu::CanReleaseLocks(TArray<FAssetData> InAssetObjectPaths) const
+{
+	const TArray<FString> Files = PackageUtils::AssetDataToFileNames(InAssetObjectPaths);
+
+	for (const FString& File : Files)
+	{
+		const FString AbsoluteFilename = FPaths::ConvertRelativePathToFull(File);
+		const auto State = FPlasticSourceControlModule::Get().GetProvider().GetStateInternal(AbsoluteFilename);
+		// If exclusively Checked Out (Locked) the lock can be released coming back to it's potential underlying "Retained" status if changes where already checked in the branch
+		if (!State->LockedBy.IsEmpty() && State->LockedId != ISourceControlState::INVALID_REVISION)
 		{
-			// Find and Unlink all packages in Content directory to allow to update them
-			PackageUtils::UnlinkPackages(ListAllPackages());
+			return true;
+		}
+	}
 
-			// Launch a custom "SyncAll" operation
-			FPlasticSourceControlProvider& Provider = FPlasticSourceControlModule::Get().GetProvider();
-			TSharedRef<FPlasticSyncAll, ESPMode::ThreadSafe> SyncOperation = ISourceControlOperation::Create<FPlasticSyncAll>();
-			const ECommandResult::Type Result = Provider.Execute(SyncOperation, TArray<FString>(), EConcurrency::Asynchronous, FSourceControlOperationComplete::CreateRaw(this, &FPlasticSourceControlMenu::OnSourceControlOperationComplete));
-			if (Result == ECommandResult::Succeeded)
-			{
-				// Display an ongoing notification during the whole operation (packages will be reloaded at the completion of the operation)
-				DisplayInProgressNotification(SyncOperation->GetInProgressString());
-			}
-			else
-			{
-				// Report failure with a notification (but nothing need to be reloaded since no local change is expected)
-				DisplayFailureNotification(SyncOperation->GetName());
-			}
+	return false;
+}
+
+bool FPlasticSourceControlMenu::CanRemoveLocks(TArray<FAssetData> InAssetObjectPaths) const
+{
+	const TArray<FString> Files = PackageUtils::AssetDataToFileNames(InAssetObjectPaths);
+
+	for (const FString& File : Files)
+	{
+		const FString AbsoluteFilename = FPaths::ConvertRelativePathToFull(File);
+		const auto State = FPlasticSourceControlModule::Get().GetProvider().GetStateInternal(AbsoluteFilename);
+		// If Locked or Retained, the lock can be removed, that is completely deleted in order to simply ignore the changes from the branch
+		if (State->LockedId != ISourceControlState::INVALID_REVISION)
+		{
+			return true;
+		}
+	}
+
+	return false;
+}
+
+void FPlasticSourceControlMenu::ExecuteReleaseLocks(TArray<FAssetData> InAssetObjectPaths)
+{
+	ExecuteUnlock(InAssetObjectPaths, false);
+}
+
+void FPlasticSourceControlMenu::ExecuteRemoveLocks(TArray<FAssetData> InAssetObjectPaths)
+{
+	ExecuteUnlock(InAssetObjectPaths, true);
+}
+
+void FPlasticSourceControlMenu::ExecuteUnlock(const TArray<FAssetData>& InAssetObjectPaths, const bool bInRemove)
+{
+	if (!Notification.IsInProgress())
+	{
+		const TArray<FString> Files = PackageUtils::AssetDataToFileNames(InAssetObjectPaths);
+
+		// Launch a custom "Release/Remove Lock" operation
+		FPlasticSourceControlProvider& Provider = FPlasticSourceControlModule::Get().GetProvider();
+		TSharedRef<FPlasticUnlock, ESPMode::ThreadSafe> UnlockOperation = ISourceControlOperation::Create<FPlasticUnlock>();
+		const ECommandResult::Type Result = Provider.Execute(UnlockOperation, Files, EConcurrency::Asynchronous, FSourceControlOperationComplete::CreateRaw(this, &FPlasticSourceControlMenu::OnSourceControlOperationComplete));
+		UnlockOperation->bRemove = bInRemove;
+		if (Result == ECommandResult::Succeeded)
+		{
+			// Display an ongoing notification during the whole operation (packages will be reloaded at the completion of the operation)
+			Notification.DisplayInProgress(UnlockOperation->GetInProgressString());
 		}
 		else
 		{
-			FMessageLog SourceControlLog("SourceControl");
-			SourceControlLog.Warning(LOCTEXT("SourceControlMenu_Sync_Unsaved", "Save All Assets before attempting to Sync!"));
-			SourceControlLog.Notify();
+			// Report failure with a notification (but nothing need to be reloaded since no local change is expected)
+			FNotification::DisplayFailure(UnlockOperation.Get());
+		}
+	}
+	else
+	{
+		FMessageLog SourceControlLog("SourceControl");
+		SourceControlLog.Warning(LOCTEXT("SourceControlMenu_InProgress", "Source control operation already in progress"));
+		SourceControlLog.Notify();
+	}
+}
+
+bool FPlasticSourceControlMenu::IsSourceControlConnected() const
+{
+	const ISourceControlProvider& Provider = ISourceControlModule::Get().GetProvider();
+	return Provider.IsEnabled() && Provider.IsAvailable();
+}
+
+void FPlasticSourceControlMenu::SyncProjectClicked()
+{
+	if (!Notification.IsInProgress())
+	{
+		// Warn the user about any unsaved assets (risk of losing work) but don't enforce saving them (reduces friction and solves some user scenario)
+		PackageUtils::SaveDirtyPackages();
+
+		// Find and Unlink all loaded packages in Content directory to allow to update them
+		PackageUtils::UnlinkPackages(PackageUtils::ListAllPackages());
+
+		// Launch a custom "SyncAll" operation
+		FPlasticSourceControlProvider& Provider = FPlasticSourceControlModule::Get().GetProvider();
+		TSharedRef<FPlasticSyncAll, ESPMode::ThreadSafe> SyncOperation = ISourceControlOperation::Create<FPlasticSyncAll>();
+		const ECommandResult::Type Result = Provider.Execute(SyncOperation, TArray<FString>(), EConcurrency::Asynchronous, FSourceControlOperationComplete::CreateRaw(this, &FPlasticSourceControlMenu::OnSyncAllOperationComplete));
+		if (Result == ECommandResult::Succeeded)
+		{
+			// Display an ongoing notification during the whole operation (packages will be reloaded at the completion of the operation)
+			Notification.DisplayInProgress(SyncOperation->GetInProgressString());
+		}
+		else
+		{
+			// Report failure with a notification (but nothing need to be reloaded since no local change is expected)
+			FNotification::DisplayFailure(SyncOperation.Get());
 		}
 	}
 	else
@@ -146,7 +287,7 @@ void FPlasticSourceControlMenu::SyncProjectClicked()
 
 void FPlasticSourceControlMenu::RevertUnchangedClicked()
 {
-	if (!OperationInProgressNotification.IsValid())
+	if (!Notification.IsInProgress())
 	{
 		// Launch a "RevertUnchanged" Operation
 		FPlasticSourceControlProvider& Provider = FPlasticSourceControlModule::Get().GetProvider();
@@ -155,12 +296,12 @@ void FPlasticSourceControlMenu::RevertUnchangedClicked()
 		if (Result == ECommandResult::Succeeded)
 		{
 			// Display an ongoing notification during the whole operation
-			DisplayInProgressNotification(RevertUnchangedOperation->GetInProgressString());
+			Notification.DisplayInProgress(RevertUnchangedOperation->GetInProgressString());
 		}
 		else
 		{
 			// Report failure with a notification
-			DisplayFailureNotification(RevertUnchangedOperation->GetName());
+			FNotification::DisplayFailure(RevertUnchangedOperation.Get());
 		}
 	}
 	else
@@ -173,39 +314,37 @@ void FPlasticSourceControlMenu::RevertUnchangedClicked()
 
 void FPlasticSourceControlMenu::RevertAllClicked()
 {
-	if (!OperationInProgressNotification.IsValid())
+	if (!Notification.IsInProgress())
 	{
 		// Ask the user before reverting all!
-		const FText DialogText(LOCTEXT("SourceControlMenu_AskRevertAll", "Revert all modifications into the workspace?"));
-		const EAppReturnType::Type Choice = FMessageDialog::Open(EAppMsgType::OkCancel, DialogText);
+		const FText AskRevertAllWarning(LOCTEXT("SourceControlMenu_AskRevertAll", "Revert all modifications into the workspace?\n"
+			"This cannot be undone."));
+		const EAppReturnType::Type Choice = FMessageDialog::Open(
+			EAppMsgCategory::Warning,
+			EAppMsgType::OkCancel, AskRevertAllWarning
+			, LOCTEXT("SourceControlMenu_AskRevertAllTitle", "Revert All?")
+		);
 		if (Choice == EAppReturnType::Ok)
 		{
-			const bool bSaved = SaveDirtyPackages();
-			if (bSaved)
+			// Warn the user about any unsaved assets (risk of losing work) but don't enforce saving them (reduces friction and solves some user scenario)
+			PackageUtils::SaveDirtyPackages();
+
+			// Find and Unlink all packages in Content directory to allow to update them
+			PackageUtils::UnlinkPackages(PackageUtils::ListAllPackages());
+
+			// Launch a "RevertAll" Operation
+			FPlasticSourceControlProvider& Provider = FPlasticSourceControlModule::Get().GetProvider();
+			TSharedRef<FPlasticRevertAll, ESPMode::ThreadSafe> RevertAllOperation = ISourceControlOperation::Create<FPlasticRevertAll>();
+			const ECommandResult::Type Result = Provider.Execute(RevertAllOperation, TArray<FString>(), EConcurrency::Asynchronous, FSourceControlOperationComplete::CreateRaw(this, &FPlasticSourceControlMenu::OnRevertAllOperationComplete));
+			if (Result == ECommandResult::Succeeded)
 			{
-				// Find and Unlink all packages in Content directory to allow to update them
-				PackageUtils::UnlinkPackages(ListAllPackages());
-
-				// Launch a "RevertAll" Operation
-				FPlasticSourceControlProvider& Provider = FPlasticSourceControlModule::Get().GetProvider();
-				TSharedRef<FPlasticRevertAll, ESPMode::ThreadSafe> RevertAllOperation = ISourceControlOperation::Create<FPlasticRevertAll>();
-				const ECommandResult::Type Result = Provider.Execute(RevertAllOperation, TArray<FString>(), EConcurrency::Asynchronous, FSourceControlOperationComplete::CreateRaw(this, &FPlasticSourceControlMenu::OnSourceControlOperationComplete));
-				if (Result == ECommandResult::Succeeded)
-				{
-					// Display an ongoing notification during the whole operation
-					DisplayInProgressNotification(RevertAllOperation->GetInProgressString());
-				}
-				else
-				{
-					// Report failure with a notification (but nothing need to be reloaded since no local change is expected)
-					DisplayFailureNotification(RevertAllOperation->GetName());
-				}
+				// Display an ongoing notification during the whole operation
+				Notification.DisplayInProgress(RevertAllOperation->GetInProgressString());
 			}
 			else
 			{
-				FMessageLog SourceControlLog("SourceControl");
-				SourceControlLog.Warning(LOCTEXT("SourceControlMenu_Sync_Unsaved", "Save All Assets before attempting to Sync!"));
-				SourceControlLog.Notify();
+				// Report failure with a notification (but nothing need to be reloaded since no local change is expected)
+				FNotification::DisplayFailure(RevertAllOperation.Get());
 			}
 		}
 	}
@@ -217,25 +356,34 @@ void FPlasticSourceControlMenu::RevertAllClicked()
 	}
 }
 
-void FPlasticSourceControlMenu::RefreshClicked()
+void FPlasticSourceControlMenu::SwitchToPartialWorkspaceClicked()
 {
-	if (!OperationInProgressNotification.IsValid())
+	if (!Notification.IsInProgress())
 	{
-		// Launch an "UpdateStatus" Operation
-		FPlasticSourceControlProvider& Provider = FPlasticSourceControlModule::Get().GetProvider();
-		TSharedRef<FUpdateStatus, ESPMode::ThreadSafe> RefreshOperation = ISourceControlOperation::Create<FUpdateStatus>();
-		// This is the flag used by the Content Browser's "Checkout" filter to trigger a full status update
-		RefreshOperation->SetGetOpenedOnly(true);
-		const ECommandResult::Type Result = Provider.Execute(RefreshOperation, TArray<FString>(), EConcurrency::Asynchronous, FSourceControlOperationComplete::CreateRaw(this, &FPlasticSourceControlMenu::OnSourceControlOperationComplete));
-		if (Result == ECommandResult::Succeeded)
-		{
-			// Display an ongoing notification during the whole operation
-			DisplayInProgressNotification(RefreshOperation->GetInProgressString());
-		}
-		else
+		// Ask the user before switching to Partial Workspace. It's not possible to switch back with local changes!
+		const FText SwitchToPartialQuestion(LOCTEXT("SourceControlMenu_AskSwitchToPartialWorkspace", "Switch to Gluon partial workspace?\n"
+			"Please note that in order to switch back to a regular workspace you will need to undo any local changes."));
+		const EAppReturnType::Type Choice = FMessageDialog::Open(
+			EAppMsgCategory::Info,
+			EAppMsgType::OkCancel, SwitchToPartialQuestion
+			, LOCTEXT("SourceControlMenu_SwitchToPartialTitle", "Switch to Gluon partial workspace?")
+		);
+		if (Choice == EAppReturnType::Ok)
 		{
-			// Report failure with a notification
-			DisplayFailureNotification(RefreshOperation->GetName());
+			// Launch a "SwitchToPartialWorkspace" Operation
+			FPlasticSourceControlProvider& Provider = FPlasticSourceControlModule::Get().GetProvider();
+			TSharedRef<FPlasticSwitchToPartialWorkspace, ESPMode::ThreadSafe> SwitchOperation = ISourceControlOperation::Create<FPlasticSwitchToPartialWorkspace>();
+			const ECommandResult::Type Result = Provider.Execute(SwitchOperation, TArray<FString>(), EConcurrency::Asynchronous, FSourceControlOperationComplete::CreateRaw(this, &FPlasticSourceControlMenu::OnSourceControlOperationComplete));
+			if (Result == ECommandResult::Succeeded)
+			{
+				// Display an ongoing notification during the whole operation
+				Notification.DisplayInProgress(SwitchOperation->GetInProgressString());
+			}
+			else
+			{
+				// Report failure with a notification
+				FNotification::DisplayFailure(SwitchOperation.Get());
+			}
 		}
 	}
 	else
@@ -246,6 +394,11 @@ void FPlasticSourceControlMenu::RefreshClicked()
 	}
 }
 
+bool FPlasticSourceControlMenu::CanSwitchToPartialWorkspace() const
+{
+	return !FPlasticSourceControlModule::Get().GetProvider().IsPartialWorkspace();
+}
+
 void FPlasticSourceControlMenu::ShowSourceControlEditorPreferences() const
 {
 	if (ISettingsModule* SettingsModule = FModuleManager::GetModulePtr<ISettingsModule>("Settings"))
@@ -273,7 +426,7 @@ void FPlasticSourceControlMenu::ShowSourceControlPlasticScmProjectSettings() con
 void FPlasticSourceControlMenu::VisitDocsURLClicked() const
 {
 	// Grab the URL from the uplugin file
-	const TSharedPtr<IPlugin> Plugin = IPluginManager::Get().FindPlugin(TEXT("PlasticSourceControl"));
+	const TSharedPtr<IPlugin> Plugin = FPlasticSourceControlModule::GetPlugin();
 	if (Plugin.IsValid())
 	{
 		FPlatformProcess::LaunchURL(*Plugin->GetDescriptor().DocsURL, NULL, NULL);
@@ -283,93 +436,50 @@ void FPlasticSourceControlMenu::VisitDocsURLClicked() const
 void FPlasticSourceControlMenu::VisitSupportURLClicked() const
 {
 	// Grab the URL from the uplugin file
-	const TSharedPtr<IPlugin> Plugin = IPluginManager::Get().FindPlugin(TEXT("PlasticSourceControl"));
+	const TSharedPtr<IPlugin> Plugin = FPlasticSourceControlModule::GetPlugin();
 	if (Plugin.IsValid())
 	{
 		FPlatformProcess::LaunchURL(*Plugin->GetDescriptor().SupportURL, NULL, NULL);
 	}
 }
 
-// Display an ongoing notification during the whole operation
-void FPlasticSourceControlMenu::DisplayInProgressNotification(const FText& InOperationInProgressString)
+void FPlasticSourceControlMenu::VisitLockRulesURLClicked(const FString InOrganizationName) const
 {
-	if (!OperationInProgressNotification.IsValid())
-	{
-		FNotificationInfo Info(InOperationInProgressString);
-		Info.bFireAndForget = false;
-		Info.ExpireDuration = 0.0f;
-		Info.FadeOutDuration = 1.0f;
-		OperationInProgressNotification = FSlateNotificationManager::Get().AddNotification(Info);
-		if (OperationInProgressNotification.IsValid())
-		{
-			OperationInProgressNotification.Pin()->SetCompletionState(SNotificationItem::CS_Pending);
-		}
-	}
+	const FString OrganizationLockRulesURL = FString::Printf(
+		TEXT("https://dashboard.unity3d.com/devops/organizations/default/plastic-scm/organizations/%s/lock-rules"),
+		*InOrganizationName
+	);
+	FPlatformProcess::LaunchURL(*OrganizationLockRulesURL, NULL, NULL);
 }
 
-// Remove the ongoing notification at the end of the operation
-void FPlasticSourceControlMenu::RemoveInProgressNotification()
+void FPlasticSourceControlMenu::OpenBranchesWindow() const
 {
-	if (OperationInProgressNotification.IsValid())
-	{
-		OperationInProgressNotification.Pin()->ExpireAndFadeout();
-		OperationInProgressNotification.Reset();
-	}
+	FPlasticSourceControlModule::Get().GetBranchesWindow().OpenTab();
 }
 
-// Display a temporary success notification at the end of the operation
-void FPlasticSourceControlMenu::DisplaySucessNotification(const FName& InOperationName)
+void FPlasticSourceControlMenu::OnSyncAllOperationComplete(const FSourceControlOperationRef& InOperation, ECommandResult::Type InResult)
 {
-	const FText NotificationText = FText::Format(
-		LOCTEXT("SourceControlMenu_Success", "{0} operation was successful!"),
-		FText::FromName(InOperationName)
-	);
-	FNotificationInfo Info(NotificationText);
-	Info.bUseSuccessFailIcons = true;
-	Info.Image = FAppStyle::GetBrush(TEXT("NotificationList.SuccessImage"));
-	FSlateNotificationManager::Get().AddNotification(Info);
-	UE_LOG(LogSourceControl, Log, TEXT("%s"), *NotificationText.ToString());
+	OnSourceControlOperationComplete(InOperation, InResult);
+
+	// Reload packages that where updated by the Sync operation (and the current map if needed)
+	TSharedRef<FPlasticSyncAll, ESPMode::ThreadSafe> Operation = StaticCastSharedRef<FPlasticSyncAll>(InOperation);
+	PackageUtils::ReloadPackages(Operation->UpdatedFiles);
 }
 
-// Display a temporary failure notification at the end of the operation
-void FPlasticSourceControlMenu::DisplayFailureNotification(const FName& InOperationName)
+void FPlasticSourceControlMenu::OnRevertAllOperationComplete(const FSourceControlOperationRef & InOperation, ECommandResult::Type InResult)
 {
-	const FText NotificationText = FText::Format(
-		LOCTEXT("SourceControlMenu_Failure", "Error: {0} operation failed!"),
-		FText::FromName(InOperationName)
-	);
-	FNotificationInfo Info(NotificationText);
-	Info.ExpireDuration = 8.0f;
-	FSlateNotificationManager::Get().AddNotification(Info);
-	UE_LOG(LogSourceControl, Error, TEXT("%s"), *NotificationText.ToString());
+	OnSourceControlOperationComplete(InOperation, InResult);
+
+	// Reload packages that where updated by the Revert operation (and the current map if needed)
+	TSharedRef<FPlasticRevertAll, ESPMode::ThreadSafe> Operation = StaticCastSharedRef<FPlasticRevertAll>(InOperation);
+	PackageUtils::ReloadPackages(Operation->UpdatedFiles);
 }
 
 void FPlasticSourceControlMenu::OnSourceControlOperationComplete(const FSourceControlOperationRef& InOperation, ECommandResult::Type InResult)
 {
-	RemoveInProgressNotification();
-
-	if (InOperation->GetName() == "SyncAll")
-	{
-		// Reload packages that where updated by the Sync operation (and the current map if needed)
-		TSharedRef<FPlasticSyncAll, ESPMode::ThreadSafe> Operation = StaticCastSharedRef<FPlasticSyncAll>(InOperation);
-		PackageUtils::ReloadPackages(Operation->UpdatedFiles);
-	}
-	else if (InOperation->GetName() == "RevertAll")
-	{
-		// Reload packages that where updated by the Revert operation (and the current map if needed)
-		TSharedRef<FPlasticRevertAll, ESPMode::ThreadSafe> Operation = StaticCastSharedRef<FPlasticRevertAll>(InOperation);
-		PackageUtils::ReloadPackages(Operation->UpdatedFiles);
-	}
+	Notification.RemoveInProgress();
 
-	// Report result with a notification
-	if (InResult == ECommandResult::Succeeded)
-	{
-		DisplaySucessNotification(InOperation->GetName());
-	}
-	else
-	{
-		DisplayFailureNotification(InOperation->GetName());
-	}
+	FNotification::DisplayResult(InOperation, InResult);
 }
 
 void FPlasticSourceControlMenu::AddMenuExtension(FToolMenuSection& Menu)
@@ -408,13 +518,14 @@ void FPlasticSourceControlMenu::AddMenuExtension(FToolMenuSection& Menu)
 	);
 
 	Menu.AddMenuEntry(
-		"PlasticRefresh",
-		LOCTEXT("PlasticRefresh",			"Refresh"),
-		LOCTEXT("PlasticRefreshTooltip",	"Update the local source control status of all files in the workspace (no expensive checks for locks or changes on other branches)."),
-		FSlateIcon(FAppStyle::GetAppStyleSetName(), "SourceControl.Actions.Refresh"),
+		"SwitchToPartialWorkspace",
+		LOCTEXT("SwitchToPartialWorkspace",			"Switch to Gluon Partial Workspace"),
+		LOCTEXT("SwitchToPartialWorkspaceTooltip",	"Update the workspace to a Gluon partial mode for a simplified workflow.\n"
+			"Allows to update and check in files individually as opposed to the whole workspace.\nIt doesn't work with branches or shelves."),
+		FSlateIcon(FAppStyle::GetAppStyleSetName(), "GenericCommands.Cut"),
 		FUIAction(
-			FExecuteAction::CreateRaw(this, &FPlasticSourceControlMenu::RefreshClicked),
-			FCanExecuteAction()
+			FExecuteAction::CreateRaw(this, &FPlasticSourceControlMenu::SwitchToPartialWorkspaceClicked),
+			FCanExecuteAction::CreateRaw(this, &FPlasticSourceControlMenu::CanSwitchToPartialWorkspace)
 		)
 	);
 
@@ -431,8 +542,8 @@ void FPlasticSourceControlMenu::AddMenuExtension(FToolMenuSection& Menu)
 
 	Menu.AddMenuEntry(
 		"SourceControlProjectSettings",
-		LOCTEXT("SourceControlProjectSettings",			"Project Settings - Source Control"),
-		LOCTEXT("SourceControlProjectSettingsTooltip",	"Open the Source Control section in the Project Settings."),
+		LOCTEXT("SourceControlProjectSettings",			"Project Settings - Revision Control"),
+		LOCTEXT("SourceControlProjectSettingsTooltip",	"Open the Revision Control section in the Project Settings."),
 		FSlateIcon(FAppStyle::GetAppStyleSetName(), "ProjectSettings.TabIcon"),
 		FUIAction(
 			FExecuteAction::CreateRaw(this, &FPlasticSourceControlMenu::ShowSourceControlProjectSettings),
@@ -442,8 +553,8 @@ void FPlasticSourceControlMenu::AddMenuExtension(FToolMenuSection& Menu)
 
 	Menu.AddMenuEntry(
 		"PlasticProjectSettings",
-		LOCTEXT("PlasticProjectSettings",			"Project Settings - Source Control - Plastic SCM"),
-		LOCTEXT("PlasticProjectSettingsTooltip",	"Open the Plastic SCM section in the Project Settings."),
+		LOCTEXT("PlasticProjectSettings",			"Project Settings - Source Control - Unity Version Control"),
+		LOCTEXT("PlasticProjectSettingsTooltip",	"Open the Unity Version Control (formerly Plastic SCM) section in the Project Settings."),
 		FSlateIcon(FAppStyle::GetAppStyleSetName(), "ProjectSettings.TabIcon"),
 		FUIAction(
 			FExecuteAction::CreateRaw(this, &FPlasticSourceControlMenu::ShowSourceControlPlasticScmProjectSettings),
@@ -464,14 +575,45 @@ void FPlasticSourceControlMenu::AddMenuExtension(FToolMenuSection& Menu)
 
 	Menu.AddMenuEntry(
 		"PlasticSupportURL",
-		LOCTEXT("PlasticSupportURL",		"Plastic SCM Support"),
-		LOCTEXT("PlasticSupportURLTooltip",	"Visit official support for Plastic SCM."),
+		LOCTEXT("PlasticSupportURL",		"Unity Version Control Support"),
+		LOCTEXT("PlasticSupportURLTooltip",	"Submit a support request for Unity Version Control (formerly Plastic SCM)."),
 		FSlateIcon(FAppStyle::GetAppStyleSetName(), "Icons.Support"),
 		FUIAction(
 			FExecuteAction::CreateRaw(this, &FPlasticSourceControlMenu::VisitSupportURLClicked),
 			FCanExecuteAction()
 		)
 	);
+
+	FString OrganizationName = FPlasticSourceControlModule::Get().GetProvider().GetCloudOrganization();
+	if (!OrganizationName.IsEmpty())
+	{
+		Menu.AddMenuEntry(
+			"PlasticLockRulesURL",
+			LOCTEXT("PlasticLockRulesURL", "Configure Lock Rules"),
+			LOCTEXT("PlasticLockRulesURLTooltip", "Navigate to lock rules configuration page in the Unity Dashboard."),
+			FSlateIcon(FAppStyle::GetAppStyleSetName(), "PropertyWindow.Locked"),
+			FUIAction(
+				FExecuteAction::CreateRaw(this, &FPlasticSourceControlMenu::VisitLockRulesURLClicked, MoveTemp(OrganizationName)),
+				FCanExecuteAction()
+			)
+		);
+	}
+
+	AddViewBranches(Menu);
+}
+
+void FPlasticSourceControlMenu::AddViewBranches(FToolMenuSection& Menu)
+{
+	Menu.AddMenuEntry(
+		TEXT("PlasticBranchesWindow"),
+		LOCTEXT("PlasticBranchesWindow", "View Branches"),
+		LOCTEXT("PlasticBranchesWindowTooltip", "Open the Branches window."),
+		FSlateIcon(FAppStyle::GetAppStyleSetName(), "SourceControl.Branch"),
+		FUIAction(
+			FExecuteAction::CreateRaw(this, &FPlasticSourceControlMenu::OpenBranchesWindow),
+			FCanExecuteAction()
+		)
+	);
 }
 
 #undef LOCTEXT_NAMESPACE
diff --git a/Engine/Plugins/Developer/PlasticSourceControl/Source/PlasticSourceControl/Private/PlasticSourceControlMenu.h b/Engine/Plugins/Developer/PlasticSourceControl/Source/PlasticSourceControl/Private/PlasticSourceControlMenu.h
--- a/Engine/Plugins/Developer/PlasticSourceControl/Source/PlasticSourceControl/Private/PlasticSourceControlMenu.h
+++ b/Engine/Plugins/Developer/PlasticSourceControl/Source/PlasticSourceControl/Private/PlasticSourceControlMenu.h
@@ -4,10 +4,13 @@
 
 #include "CoreMinimal.h"
 #include "ISourceControlProvider.h"
+#include "Notification.h"
+#include "AssetRegistry/AssetData.h"
 
+class FMenuBuilder;
 struct FToolMenuSection;
 
-/** Plastic SCM extension of the Source Control toolbar menu */
+/** Unity Version Control extension of the Source Control toolbar menu */
 class FPlasticSourceControlMenu
 {
 public:
@@ -18,34 +21,55 @@ public:
 	void SyncProjectClicked();
 	void RevertUnchangedClicked();
 	void RevertAllClicked();
-	void RefreshClicked();
+	void SwitchToPartialWorkspaceClicked();
+	bool CanSwitchToPartialWorkspace() const;
 	void ShowSourceControlEditorPreferences() const;
 	void ShowSourceControlProjectSettings() const;
 	void ShowSourceControlPlasticScmProjectSettings() const;
 	void VisitDocsURLClicked() const;
 	void VisitSupportURLClicked() const;
+	void VisitLockRulesURLClicked(const FString InOrganizationName) const;
+	void OpenBranchesWindow() const;
 
 private:
 	bool IsSourceControlConnected() const;
 
-	bool				SaveDirtyPackages();
-	TArray<FString>		ListAllPackages();
-
 	void AddMenuExtension(FToolMenuSection& Menu);
+	void AddViewBranches(FToolMenuSection& Menu);
 
-	void DisplayInProgressNotification(const FText& InOperationInProgressString);
-	void RemoveInProgressNotification();
-	void DisplaySucessNotification(const FName& InOperationName);
-	void DisplayFailureNotification(const FName& InOperationName);
+	/** Extends the UE5 toolbar with a status bar widget to display the current branch and open the branch tab */
+	void ExtendToolbarWithStatusBarWidget();
 
-private:
+	/** Extends the main Revision Control menu from the toolbar at the bottom-right. */
+	void ExtendRevisionControlMenu();
+	/** Extends the content browser asset context menu with Admin revision control options. */
+	void ExtendAssetContextMenu();
+	/** Called to generate concert asset context menu. */
+	void GeneratePlasticAssetContextMenu(FMenuBuilder& MenuBuilder, TArray<FAssetData> InAssetObjectPaths);
+
+	bool CanRemoveLocks(TArray<FAssetData> InAssetObjectPaths) const;
+	bool CanReleaseLocks(TArray<FAssetData> InAssetObjectPaths) const;
+	void ExecuteRemoveLocks(TArray<FAssetData> InAssetObjectPaths);
+	void ExecuteReleaseLocks(TArray<FAssetData> InAssetObjectPaths);
+	void ExecuteUnlock(const TArray<FAssetData>& InAssetObjectPaths, const bool bInRemove);
 
+private:
 	/** Tracks if the menu extension has been registered with the editor or not */
 	bool bHasRegistered = false;
 
-	/** Current source control operation from extended menu if any */
-	TWeakPtr<class SNotificationItem> OperationInProgressNotification;
+	/** Ongoing notification for a long-running asynchronous source control operation, if any */
+	FNotification Notification;
+
+	/** Name of the menu extension going into the global Revision Control (on the toolbar at the bottom-right) */
+	static FName UnityVersionControlMainMenuOwnerName;
+	/** Name of the asset context menu extension for admin actions over Locks */
+	static FName UnityVersionControlAssetContextLocksMenuOwnerName;
+	/** Name of status bar extension to display the current branch  */
+	static FName UnityVersionControlStatusBarMenuOwnerName;
 
-	/** Delegate called when a source control operation has completed */
+	/** Delegates called when a source control operation has completed */
+	void OnSyncAllOperationComplete(const FSourceControlOperationRef& InOperation, ECommandResult::Type InResult);
+	void OnRevertAllOperationComplete(const FSourceControlOperationRef& InOperation, ECommandResult::Type InResult);
+	/** Generic delegate and notification handler */
 	void OnSourceControlOperationComplete(const FSourceControlOperationRef& InOperation, ECommandResult::Type InResult);
 };
diff --git a/Engine/Plugins/Developer/PlasticSourceControl/Source/PlasticSourceControl/Private/PlasticSourceControlModule.cpp b/Engine/Plugins/Developer/PlasticSourceControl/Source/PlasticSourceControl/Private/PlasticSourceControlModule.cpp
--- a/Engine/Plugins/Developer/PlasticSourceControl/Source/PlasticSourceControl/Private/PlasticSourceControlModule.cpp
+++ b/Engine/Plugins/Developer/PlasticSourceControl/Source/PlasticSourceControl/Private/PlasticSourceControlModule.cpp
@@ -3,8 +3,11 @@
 #include "PlasticSourceControlModule.h"
 
 #include "IPlasticSourceControlWorker.h"
+
+#include "Interfaces/IPluginManager.h"
 #include "Features/IModularFeatures.h"
 #include "Misc/App.h"
+#include "Modules/ModuleManager.h"
 
 #define LOCTEXT_NAMESPACE "PlasticSourceControl"
 
@@ -15,6 +18,9 @@ void FPlasticSourceControlModule::StartupModule()
 
 	// Bind our source control provider to the editor
 	IModularFeatures::Get().RegisterModularFeature("SourceControl", &PlasticSourceControlProvider);
+
+	/// Register our tab Window here as it needs to be ready for the editor to reload at startup
+	PlasticSourceControlBranchesWindow.Register();
 }
 
 void FPlasticSourceControlModule::ShutdownModule()
@@ -22,10 +28,27 @@ void FPlasticSourceControlModule::ShutdownModule()
 	// shut down the provider, as this module is going away
 	PlasticSourceControlProvider.Close();
 
+	PlasticSourceControlBranchesWindow.Unregister();
+
 	// unbind provider from editor
 	IModularFeatures::Get().UnregisterModularFeature("SourceControl", &PlasticSourceControlProvider);
 }
 
+FPlasticSourceControlModule& FPlasticSourceControlModule::Get()
+{
+	return FModuleManager::GetModuleChecked<FPlasticSourceControlModule>("PlasticSourceControl");
+}
+
+bool FPlasticSourceControlModule::IsLoaded()
+{
+	return FModuleManager::Get().IsModuleLoaded("PlasticSourceControl");
+}
+
+const TSharedPtr<IPlugin> FPlasticSourceControlModule::GetPlugin()
+{
+	return IPluginManager::Get().FindPlugin(TEXT("PlasticSourceControl"));;
+}
+
 IMPLEMENT_MODULE(FPlasticSourceControlModule, PlasticSourceControl);
 
 #undef LOCTEXT_NAMESPACE
diff --git a/Engine/Plugins/Developer/PlasticSourceControl/Source/PlasticSourceControl/Private/PlasticSourceControlModule.h b/Engine/Plugins/Developer/PlasticSourceControl/Source/PlasticSourceControl/Private/PlasticSourceControlModule.h
--- a/Engine/Plugins/Developer/PlasticSourceControl/Source/PlasticSourceControl/Private/PlasticSourceControlModule.h
+++ b/Engine/Plugins/Developer/PlasticSourceControl/Source/PlasticSourceControl/Private/PlasticSourceControlModule.h
@@ -4,11 +4,14 @@
 
 #include "CoreMinimal.h"
 #include "Modules/ModuleInterface.h"
-#include "Modules/ModuleManager.h"
+
 #include "PlasticSourceControlProvider.h"
+#include "PlasticSourceControlWorkspaceCreation.h"
+
+#include "PlasticSourceControlBranchesWindow.h"
 
 /**
- * PlasticSourceControl is the official Plastic SCM Source Control Plugin for Unreal Engine
+ * PlasticSourceControl is the official Unity Version Control Plugin for Unreal Engine
  *
  * Written and contributed by [REDACTED_PERSON] ([REDACTED_EMAIL]) for Codice Software
  */
@@ -29,18 +32,43 @@ public:
 		return PlasticSourceControlProvider;
 	}
 
+	/** Access the controller to create a new workspace */
+	FPlasticSourceControlWorkspaceCreation& GetWorkspaceCreation()
+	{
+		return PlasticSourceControlWorkspaceCreation;
+	}
+
+	FPlasticSourceControlBranchesWindow& GetBranchesWindow()
+	{
+		return PlasticSourceControlBranchesWindow;
+	}
+
 	/**
 	 * Singleton-like access to this module's interface.  This is just for convenience!
 	 * Beware of calling this during the shutdown phase, though.  Your module might have been unloaded already.
 	 *
-	 * @return Returns singleton instance, loading the module on demand if needed
+	 * @return Returns singleton instance, asserts if the module is not loaded yet or unloaded already.
 	 */
-	static inline FPlasticSourceControlModule& Get()
-	{
-		return FModuleManager::GetModuleChecked<FPlasticSourceControlModule>("PlasticSourceControl");
-	}
+	static FPlasticSourceControlModule& Get();
+
+	/**
+	 * Checks whether the module is currently loaded.
+	 */
+	static bool IsLoaded();
+
+	/**
+	 * Finds information of the plugin.
+	 *
+	 * @return	 Pointer to the plugin's information, or nullptr.
+	 */
+	static const TSharedPtr<class IPlugin> GetPlugin();
 
 private:
 	/** The Plastic source control provider */
 	FPlasticSourceControlProvider PlasticSourceControlProvider;
-};
\ No newline at end of file
+
+	FPlasticSourceControlBranchesWindow PlasticSourceControlBranchesWindow;
+
+	/** Logic to create a new workspace */
+	FPlasticSourceControlWorkspaceCreation PlasticSourceControlWorkspaceCreation;
+};
