PR #11295: [Source Control] Update Plastic SCM provider to 1.9.0 (from 1.8.0)

diff --git a/Engine/Plugins/Developer/PlasticSourceControl/Source/PlasticSourceControl/Private/PlasticSourceControlOperations.cpp b/Engine/Plugins/Developer/PlasticSourceControl/Source/PlasticSourceControl/Private/PlasticSourceControlOperations.cpp
--- a/Engine/Plugins/Developer/PlasticSourceControl/Source/PlasticSourceControl/Private/PlasticSourceControlOperations.cpp
+++ b/Engine/Plugins/Developer/PlasticSourceControl/Source/PlasticSourceControl/Private/PlasticSourceControlOperations.cpp
@@ -3,10 +3,13 @@
 #include "PlasticSourceControlOperations.h"
 
 #include "PackageUtils.h"
+#include "PlasticSourceControlBranch.h"
 #include "PlasticSourceControlCommand.h"
 #include "PlasticSourceControlModule.h"
+#include "PlasticSourceControlParsers.h"
 #include "PlasticSourceControlProvider.h"
 #include "PlasticSourceControlSettings.h"
+#include "PlasticSourceControlShell.h"
 #include "PlasticSourceControlState.h"
 #include "PlasticSourceControlUtils.h"
 #include "PlasticSourceControlVersions.h"
@@ -17,7 +20,9 @@
 #include "AssetRegistry/AssetRegistryModule.h"
 #include "Async/Async.h"
 #include "HAL/FileManager.h"
+#include "Logging/MessageLog.h"
 #include "Misc/Paths.h"
+#include "Modules/ModuleManager.h"
 
 #include "ISourceControlModule.h"
 #include "SourceControlOperations.h"
@@ -41,6 +46,14 @@ void IPlasticSourceControlWorker::RegisterWorkers(FPlasticSourceControlProvider&
 	PlasticSourceControlProvider.RegisterWorker("Revert", FGetPlasticSourceControlWorker::CreateStatic(&InstantiateWorker<FPlasticRevertWorker>));
 	PlasticSourceControlProvider.RegisterWorker("RevertUnchanged", FGetPlasticSourceControlWorker::CreateStatic(&InstantiateWorker<FPlasticRevertUnchangedWorker>));
 	PlasticSourceControlProvider.RegisterWorker("RevertAll", FGetPlasticSourceControlWorker::CreateStatic(&InstantiateWorker<FPlasticRevertAllWorker>));
+	PlasticSourceControlProvider.RegisterWorker("SwitchToPartialWorkspace", FGetPlasticSourceControlWorker::CreateStatic(&InstantiateWorker<FPlasticSwitchToPartialWorkspaceWorker>));
+	PlasticSourceControlProvider.RegisterWorker("Unlock", FGetPlasticSourceControlWorker::CreateStatic(&InstantiateWorker<FPlasticUnlockWorker>));
+	PlasticSourceControlProvider.RegisterWorker("GetBranches", FGetPlasticSourceControlWorker::CreateStatic(&InstantiateWorker<FPlasticGetBranchesWorker>));
+	PlasticSourceControlProvider.RegisterWorker("Switch", FGetPlasticSourceControlWorker::CreateStatic(&InstantiateWorker<FPlasticSwitchToBranchWorker>));
+	PlasticSourceControlProvider.RegisterWorker("Merge", FGetPlasticSourceControlWorker::CreateStatic(&InstantiateWorker<FPlasticMergeBranchWorker>));
+	PlasticSourceControlProvider.RegisterWorker("CreateBranch", FGetPlasticSourceControlWorker::CreateStatic(&InstantiateWorker<FPlasticCreateBranchWorker>));
+	PlasticSourceControlProvider.RegisterWorker("RenameBranch", FGetPlasticSourceControlWorker::CreateStatic(&InstantiateWorker<FPlasticRenameBranchWorker>));
+	PlasticSourceControlProvider.RegisterWorker("DeleteBranches", FGetPlasticSourceControlWorker::CreateStatic(&InstantiateWorker<FPlasticDeleteBranchesWorker>));
 	PlasticSourceControlProvider.RegisterWorker("MakeWorkspace", FGetPlasticSourceControlWorker::CreateStatic(&InstantiateWorker<FPlasticMakeWorkspaceWorker>));
 	PlasticSourceControlProvider.RegisterWorker("Sync", FGetPlasticSourceControlWorker::CreateStatic(&InstantiateWorker<FPlasticSyncWorker>));
 	PlasticSourceControlProvider.RegisterWorker("SyncAll", FGetPlasticSourceControlWorker::CreateStatic(&InstantiateWorker<FPlasticSyncWorker>));
@@ -57,6 +70,9 @@ void IPlasticSourceControlWorker::RegisterWorkers(FPlasticSourceControlProvider&
 	PlasticSourceControlProvider.RegisterWorker("Shelve", FGetPlasticSourceControlWorker::CreateStatic(&InstantiateWorker<FPlasticShelveWorker>));
 	PlasticSourceControlProvider.RegisterWorker("Unshelve", FGetPlasticSourceControlWorker::CreateStatic(&InstantiateWorker<FPlasticUnshelveWorker>));
 	PlasticSourceControlProvider.RegisterWorker("DeleteShelved", FGetPlasticSourceControlWorker::CreateStatic(&InstantiateWorker<FPlasticDeleteShelveWorker>));
+
+	PlasticSourceControlProvider.RegisterWorker("GetChangelistDetails", FGetPlasticSourceControlWorker::CreateStatic(&InstantiateWorker<FPlasticGetChangelistDetailsWorker>));
+	PlasticSourceControlProvider.RegisterWorker("GetFile", FGetPlasticSourceControlWorker::CreateStatic(&InstantiateWorker<FPlasticGetFileWorker>));
 }
 
 
@@ -67,7 +83,12 @@ FName FPlasticRevertUnchanged::GetName() const
 
 FText FPlasticRevertUnchanged::GetInProgressString() const
 {
-	return LOCTEXT("SourceControl_RevertUnchanged", "Reverting unchanged file(s) in Source Control...");
+	return LOCTEXT("SourceControl_RevertUnchanged", "Reverting unchanged file(s) in Revision Control...");
+}
+
+FName FPlasticSyncAll::GetName() const
+{
+	return "SyncAll";
 }
 
 FName FPlasticRevertAll::GetName() const
@@ -77,7 +98,7 @@ FName FPlasticRevertAll::GetName() const
 
 FText FPlasticRevertAll::GetInProgressString() const
 {
-	return LOCTEXT("SourceControl_RevertAll", "Reverting checked-out file(s) in Source Control...");
+	return LOCTEXT("SourceControl_RevertAll", "Reverting checked-out file(s) in Revision Control...");
 }
 
 FName FPlasticMakeWorkspace::GetName() const
@@ -87,9 +108,94 @@ FName FPlasticMakeWorkspace::GetName() const
 
 FText FPlasticMakeWorkspace::GetInProgressString() const
 {
-	return LOCTEXT("SourceControl_MakeWorkspace", "Create a new Repository and initialize the Workspace");
+	return LOCTEXT("SourceControl_MakeWorkspace", "Creating a new Repository and Workspace...");
+}
+
+FName FPlasticSwitchToPartialWorkspace::GetName() const
+{
+	return "SwitchToPartialWorkspace";
+}
+
+FText FPlasticSwitchToPartialWorkspace::GetInProgressString() const
+{
+	return LOCTEXT("SourceControl_SwitchToPartialWorkspace", "Switching to a Partial/Gluon Workspace...");
+}
+
+FName FPlasticUnlock::GetName() const
+{
+	return "Unlock";
+}
+
+FText FPlasticUnlock::GetInProgressString() const
+{
+	if (bRemove)
+		return LOCTEXT("SourceControl_Unlock_Remove", "Removing Lock(s)...");
+	else
+		return LOCTEXT("SourceControl_Unlock_Release", "Releasing Lock(s)...");
+}
+
+FName FPlasticGetBranches::GetName() const
+{
+	return "GetBranches";
+}
+
+FText FPlasticGetBranches::GetInProgressString() const
+{
+	return LOCTEXT("SourceControl_GetBranches", "Getting the list of branches...");
+}
+
+FName FPlasticSwitchToBranch::GetName() const
+{
+	return "Switch";
+}
+
+FText FPlasticSwitchToBranch::GetInProgressString() const
+{
+	return LOCTEXT("SourceControl_SwitchToBranch", "Switching the workspace to another branch...");
+}
+
+FName FPlasticMergeBranch::GetName() const
+{
+	return "Merge";
+}
+
+FText FPlasticMergeBranch::GetInProgressString() const
+{
+	return LOCTEXT("SourceControl_MergeBranch", "Merging branch to the current one...");
+}
+
+FName FPlasticCreateBranch::GetName() const
+{
+	return "CreateBranch";
+}
+
+FText FPlasticCreateBranch::GetInProgressString() const
+{
+	return LOCTEXT("SourceControl_CreateBranch", "Creating new child branch...");
+}
+
+
+FName FPlasticRenameBranch::GetName() const
+{
+	return "RenameBranch";
+}
+
+FText FPlasticRenameBranch::GetInProgressString() const
+{
+	return LOCTEXT("SourceControl_RenameBranch", "Renaming a branch...");
 }
 
+FName FPlasticDeleteBranches::GetName() const
+{
+	return "DeleteBranches";
+}
+
+FText FPlasticDeleteBranches::GetInProgressString() const
+{
+	return LOCTEXT("SourceControl_DeleteBranches", "Deleting branch(es)...");
+}
+
+
 static bool AreAllFiles(const TArray<FString>& InFiles)
 {
 	for (const FString& File : InFiles)
@@ -119,9 +225,8 @@ bool FPlasticConnectWorker::Execute(FPlasticSourceControlCommand& InCommand)
 		InCommand.bCommandSuccessful = PlasticSourceControlUtils::GetWorkspaceName(InCommand.PathToWorkspaceRoot, InCommand.WorkspaceName, InCommand.ErrorMessages);
 		if (InCommand.bCommandSuccessful)
 		{
-			// Get repository, server URL, branch and current changeset number
-			// Note: this initiates the connection to the server and issue network calls, so we don't need an explicit 'checkconnection'
-			InCommand.bCommandSuccessful = PlasticSourceControlUtils::GetWorkspaceInformation(InCommand.ChangesetNumber, InCommand.RepositoryName, InCommand.ServerUrl, InCommand.BranchName, InCommand.ErrorMessages);
+			// Get current branch, repository, server URL and check the connection
+			InCommand.bCommandSuccessful = PlasticSourceControlUtils::RunCheckConnection(InCommand.BranchName, InCommand.RepositoryName, InCommand.ServerUrl, InCommand.InfoMessages, InCommand.ErrorMessages);
 			if (InCommand.bCommandSuccessful)
 			{
 				InCommand.InfoMessages.Add(TEXT("Connected successfully"));
@@ -129,11 +234,12 @@ bool FPlasticConnectWorker::Execute(FPlasticSourceControlCommand& InCommand)
 				// Now update the status of assets in the Content directory
 				// but only on real (re-)connection (but not each time Login() is called by Rename or Fixup Redirector command to check connection)
 				// and only if enabled in the settings
-				if (!GetProvider().IsAvailable() && GetProvider().AccessSettings().GetUpdateStatusAtStartup())
+				if (!PlasticSourceControlShell::GetShellIsWarmedUp() && GetProvider().AccessSettings().GetUpdateStatusAtStartup())
 				{
+					PlasticSourceControlShell::SetShellIsWarmedUp();
 					TArray<FString> ContentDir;
 					ContentDir.Add(FPaths::ConvertRelativePathToFull(FPaths::ProjectContentDir()));
-					PlasticSourceControlUtils::RunUpdateStatus(ContentDir, false, InCommand.ErrorMessages, States, InCommand.ChangesetNumber, InCommand.BranchName);
+					PlasticSourceControlUtils::RunUpdateStatus(ContentDir, PlasticSourceControlUtils::EStatusSearchType::All, false, InCommand.ErrorMessages, States, InCommand.ChangesetNumber, InCommand.BranchName);
 				}
 			}
 			else
@@ -144,7 +250,7 @@ bool FPlasticConnectWorker::Execute(FPlasticSourceControlCommand& InCommand)
 				}
 				else
 				{
-					const FText ErrorText(LOCTEXT("FailedToConnect", "Failed to connect to the Plastic SCM server."));
+					const FText ErrorText(LOCTEXT("FailedToConnect", "Failed to connect to the Unity Version Control (formerly Plastic SCM) server."));
 					Operation->SetErrorText(ErrorText);
 					InCommand.ErrorMessages.Add(ErrorText.ToString());
 				}
@@ -152,14 +258,14 @@ bool FPlasticConnectWorker::Execute(FPlasticSourceControlCommand& InCommand)
 		}
 		else
 		{
-			const FText ErrorText(LOCTEXT("NotAPlasticRepository", "Failed to enable Plastic SCM source control. You need to create a Plastic SCM workspace for the project first."));
+			const FText ErrorText(LOCTEXT("NotAPlasticRepository", "Failed to enable Unity Version Control (formerly Plastic SCM). You need to create a workspace for the project first."));
 			Operation->SetErrorText(ErrorText);
 			InCommand.ErrorMessages.Add(ErrorText.ToString());
 		}
 	}
 	else
 	{
-		const FText ErrorText(LOCTEXT("PlasticScmCliUnavaillable", "Failed to launch Plastic SCM 'cm' command line tool. You need to install it and make sure that 'cm' is on the Path and correctly configured."));
+		const FText ErrorText(LOCTEXT("PlasticScmCliUnavaillable", "Failed to launch Unity Version Control (formerly Plastic SCM) 'cm' command line tool. You need to install it and make sure it is correctly configured with your credentials."));
 		Operation->SetErrorText(ErrorText);
 		InCommand.ErrorMessages.Add(ErrorText.ToString());
 	}
@@ -185,15 +291,8 @@ static void UpdateChangelistState(FPlasticSourceControlProvider& SCCProvider, co
 
 		for (const FPlasticSourceControlState& InState : InStates)
 		{
-			// cm cannot yet handle local modifications in changelists, only the GUI can
-			if (   InState.WorkspaceState != EWorkspaceState::CheckedOut
-				&& InState.WorkspaceState != EWorkspaceState::Added
-				&& InState.WorkspaceState != EWorkspaceState::Deleted
-				&& InState.WorkspaceState != EWorkspaceState::Copied
-				&& InState.WorkspaceState != EWorkspaceState::Moved
-				&& InState.WorkspaceState != EWorkspaceState::Conflicted	// In source control, waiting for merged
-				&& InState.WorkspaceState != EWorkspaceState::Replaced		// In source control, merged, waiting for checkin to conclude the merge
-				)
+			// Note: cannot use IsModified() because cm cannot yet handle local modifications in changelists, only the GUI can
+			if (!InState.IsCheckedOutImplementation())
 			{
 				continue;
 			}
@@ -220,8 +319,7 @@ bool FPlasticCheckOutWorker::Execute(FPlasticSourceControlCommand& InCommand)
 
 	check(InCommand.Operation->GetName() == GetName());
 
-	// Detect special case for a partial checkout (CS:-1 in Gluon mode)!
-	if (-1 != InCommand.ChangesetNumber)
+	if (!GetProvider().IsPartialWorkspace())
 	{
 		InCommand.bCommandSuccessful = PlasticSourceControlUtils::RunCommand(TEXT("checkout"), TArray<FString>(), InCommand.Files, InCommand.InfoMessages, InCommand.ErrorMessages);
 	}
@@ -231,7 +329,7 @@ bool FPlasticCheckOutWorker::Execute(FPlasticSourceControlCommand& InCommand)
 	}
 
 	// now update the status of our files
-	PlasticSourceControlUtils::RunUpdateStatus(InCommand.Files, false, InCommand.ErrorMessages, States, InCommand.ChangesetNumber, InCommand.BranchName);
+	PlasticSourceControlUtils::RunUpdateStatus(InCommand.Files, PlasticSourceControlUtils::EStatusSearchType::ControlledOnly, false, InCommand.ErrorMessages, States, InCommand.ChangesetNumber, InCommand.BranchName);
 
 	return InCommand.bCommandSuccessful;
 }
@@ -260,7 +358,7 @@ bool DeleteChangelist(const FPlasticSourceControlProvider& PlasticSourceControlP
 	{
 		Parameters.Add(TEXT("delete"));
 		const FScopedTempFile ChangelistNameFile(InChangelist.GetName());
-		Parameters.Add(FString::Printf(TEXT("--namefile=\"%s\""), *FPaths::ConvertRelativePathToFull(ChangelistNameFile.GetFilename())));
+		Parameters.Add(FString::Printf(TEXT("--namefile=\"%s\""), *ChangelistNameFile.GetFilename()));
 		UE_LOG(LogSourceControl, Verbose, TEXT("DeleteChangelist(%s)"), *InChangelist.GetName());
 		return PlasticSourceControlUtils::RunCommand(TEXT("changelist"), Parameters, Files, OutResults, OutErrorMessages);
 	}
@@ -278,32 +376,6 @@ TArray<FString> FileNamesFromFileStates(const TArray<FSourceControlStateRef>& In
 	return Files;
 }
 
-/// Parse checkin result, usually looking like "Created changeset cs:8@br:/main@MyProject@SRombauts@cloud (mount:'/')"
-static FText ParseCheckInResults(const TArray<FString>& InResults)
-{
-	if (InResults.Num() > 0)
-	{
-		static const FString ChangesetPrefix(TEXT("Created changeset "));
-		if (InResults.Last().StartsWith(ChangesetPrefix))
-		{
-			FString ChangesetString;
-			static const FString BranchPrefix(TEXT("@br:"));
-			const int32 BranchIndex = InResults.Last().Find(BranchPrefix, ESearchCase::CaseSensitive);
-			if (BranchIndex > INDEX_NONE)
-			{
-				ChangesetString = InResults.Last().Mid(ChangesetPrefix.Len(), BranchIndex - ChangesetPrefix.Len());
-			}
-			return FText::Format(LOCTEXT("SubmitMessage", "Submitted changeset {0}"), FText::FromString(ChangesetString));
-		}
-		else
-		{
-			return FText::FromString(InResults.Last());
-		}
-	}
-	return FText();
-}
-
-
 TArray<FString> GetFilesFromCommand(FPlasticSourceControlProvider& PlasticSourceControlProvider, FPlasticSourceControlCommand& InCommand)
 {
 	TArray<FString> Files;
@@ -356,9 +428,8 @@ bool FPlasticCheckInWorker::Execute(FPlasticSourceControlCommand& InCommand)
 		if (!CommitMsgFile.GetFilename().IsEmpty())
 		{
 			TArray<FString> Parameters;
-			Parameters.Add(FString::Printf(TEXT("--commentsfile=\"%s\""), *FPaths::ConvertRelativePathToFull(CommitMsgFile.GetFilename())));
-			// Detect special case for a partial checkout (CS:-1 in Gluon mode)!
-			if (-1 != InCommand.ChangesetNumber)
+			Parameters.Add(FString::Printf(TEXT("--commentsfile=\"%s\""), *CommitMsgFile.GetFilename()));
+			if (!GetProvider().IsPartialWorkspace())
 			{
 				Parameters.Add(TEXT("--all")); // Also files Changed (not CheckedOut) and Moved/Deleted Locally
 			//  NOTE: --update added as #23 but removed as #32 because most assets are locked by the Unreal Editor
@@ -372,18 +443,7 @@ bool FPlasticCheckInWorker::Execute(FPlasticSourceControlCommand& InCommand)
 			}
 			if (InCommand.bCommandSuccessful)
 			{
-				// Remove any deleted files from status cache
-				TArray<TSharedRef<ISourceControlState, ESPMode::ThreadSafe>> LocalStates;
-				GetProvider().GetState(Files, LocalStates, EStateCacheUsage::Use);
-				for (const auto& State : LocalStates)
-				{
-					if (State->IsDeleted())
-					{
-						GetProvider().RemoveFileFromCache(State->GetFilename());
-					}
-				}
-
-				Operation->SetSuccessMessage(ParseCheckInResults(InCommand.InfoMessages));
+				Operation->SetSuccessMessage(PlasticSourceControlParsers::ParseCheckInResults(InCommand.InfoMessages));
 				UE_LOG(LogSourceControl, Log, TEXT("CheckIn successful"));
 			}
 
@@ -395,7 +455,7 @@ bool FPlasticCheckInWorker::Execute(FPlasticSourceControlCommand& InCommand)
 		}
 
 		// now update the status of our files
-		PlasticSourceControlUtils::RunUpdateStatus(Files, false, InCommand.ErrorMessages, States, InCommand.ChangesetNumber, InCommand.BranchName);
+		PlasticSourceControlUtils::RunUpdateStatus(Files, PlasticSourceControlUtils::EStatusSearchType::ControlledOnly, false, InCommand.ErrorMessages, States, InCommand.ChangesetNumber, InCommand.BranchName);
 	}
 	else
 	{
@@ -447,6 +507,16 @@ bool FPlasticCheckInWorker::UpdateStates()
 		}
 	}
 
+	// Remove any deleted files from status cache
+	for (const auto& State : States)
+	{
+		// Note: a file that was deleted and submitted now appears as "Private", that is, not in Source Control anymore
+		if (!State.IsSourceControlled())
+		{
+			GetProvider().RemoveFileFromCache(State.GetFilename());
+		}
+	}
+
 	return PlasticSourceControlUtils::UpdateCachedStates(MoveTemp(States));
 }
 
@@ -479,8 +549,7 @@ bool FPlasticMarkForAddWorker::Execute(FPlasticSourceControlCommand& InCommand)
 		{
 			Parameters.Add(TEXT("-R"));	// needed only at the time of workspace creation, to add directories recursively
 		}
-		// Detect special case for a partial checkout (CS:-1 in Gluon mode)!
-		if (-1 != InCommand.ChangesetNumber)
+		if (!GetProvider().IsPartialWorkspace())
 		{
 			InCommand.bCommandSuccessful = PlasticSourceControlUtils::RunCommand(TEXT("add"), Parameters, InCommand.Files, InCommand.InfoMessages, InCommand.ErrorMessages);
 		}
@@ -490,7 +559,7 @@ bool FPlasticMarkForAddWorker::Execute(FPlasticSourceControlCommand& InCommand)
 		}
 
 		// now update the status of our files
-		PlasticSourceControlUtils::RunUpdateStatus(InCommand.Files, false, InCommand.ErrorMessages, States, InCommand.ChangesetNumber, InCommand.BranchName);
+		PlasticSourceControlUtils::RunUpdateStatus(InCommand.Files, PlasticSourceControlUtils::EStatusSearchType::ControlledOnly, false, InCommand.ErrorMessages, States, InCommand.ChangesetNumber, InCommand.BranchName);
 	}
 	else
 	{
@@ -521,8 +590,7 @@ bool FPlasticDeleteWorker::Execute(FPlasticSourceControlCommand& InCommand)
 
 	check(InCommand.Operation->GetName() == GetName());
 
-	// Detect special case for a partial checkout (CS:-1 in Gluon mode)!
-	if (-1 != InCommand.ChangesetNumber)
+	if (!GetProvider().IsPartialWorkspace())
 	{
 		InCommand.bCommandSuccessful = PlasticSourceControlUtils::RunCommand(TEXT("remove"), TArray<FString>(), InCommand.Files, InCommand.InfoMessages, InCommand.ErrorMessages);
 	}
@@ -532,7 +600,7 @@ bool FPlasticDeleteWorker::Execute(FPlasticSourceControlCommand& InCommand)
 	}
 
 	// now update the status of our files
-	PlasticSourceControlUtils::RunUpdateStatus(InCommand.Files, false, InCommand.ErrorMessages, States, InCommand.ChangesetNumber, InCommand.BranchName);
+	PlasticSourceControlUtils::RunUpdateStatus(InCommand.Files, PlasticSourceControlUtils::EStatusSearchType::ControlledOnly, false, InCommand.ErrorMessages, States, InCommand.ChangesetNumber, InCommand.BranchName);
 
 	return InCommand.bCommandSuccessful;
 }
@@ -558,54 +626,101 @@ bool FPlasticRevertWorker::Execute(FPlasticSourceControlCommand& InCommand)
 
 	check(InCommand.Operation->GetName() == GetName());
 	TSharedRef<FRevert, ESPMode::ThreadSafe> Operation = StaticCastSharedRef<FRevert>(InCommand.Operation);
+	const bool bIsSoftRevert = Operation->IsSoftRevert();
+	if (bIsSoftRevert && GetProvider().GetPlasticScmVersion() < PlasticSourceControlVersions::UndoCheckoutKeepChanges)
+	{
+		// If a soft revert is requested but not supported by the version of Unity Version Control, warn the user and stop
+		FText FailureText = FText::FormatOrdered(
+			LOCTEXT("PlasticUndoKeepChangesVersionError", "Unity Version Control {0} cannot keep changes when undoing the checkout of the selected files. Update to version {1} or above."),
+			FText::FromString(*GetProvider().GetPlasticScmVersion().String),
+			FText::FromString(*PlasticSourceControlVersions::UndoCheckoutKeepChanges.String));
 
-	TArray<FString> Files = GetFilesFromCommand(GetProvider(), InCommand);
+		AsyncTask(ENamedThreads::GameThread, [FailureText]
+		{
+			FMessageLog SourceControlLog("SourceControl");
+			SourceControlLog.Error(FailureText);
+			SourceControlLog.Notify();
+		});
 
-	for (int i = 0; i < Files.Num(); i++) // Required for loop on index since we are adding to the Files array as we go
-	{
-		const FString& File = Files[i];
+		return false;
+	}
 
+	const TArray<FString> Files = GetFilesFromCommand(GetProvider(), InCommand);
+
+	TArray<FString> LocallyChangedFiles;
+	TArray<FString> CheckedOutFiles;
+
+	for (const FString& File : Files)
+	{
 		const TSharedRef<const FPlasticSourceControlState, ESPMode::ThreadSafe> State = GetProvider().GetStateInternal(File);
 
-		if (State->WorkspaceState == EWorkspaceState::Moved)
+		if (State->WorkspaceState == EWorkspaceState::Changed)
 		{
-			const FString& MovedFrom = State->MovedFrom;
+			LocallyChangedFiles.Add(State->LocalFilename);
+		}
+		else
+		{
+			CheckedOutFiles.Add(State->LocalFilename);
+			// in case of a Moved/Renamed, find the rename origin to revert both at once
+			if (State->WorkspaceState == EWorkspaceState::Moved)
+			{
+				const FString& MovedFrom = State->MovedFrom;
+
+				// In case of a file Moved/Renamed, consider the rename Origin (where there is now a Redirector file Added)
+				// and add it to the list of files to revert (only if it is not already in) to revert both at once
+				if (!CheckedOutFiles.FindByPredicate([&MovedFrom](const FString& File) { return File.Equals(MovedFrom, ESearchCase::IgnoreCase); }))
+				{
+					CheckedOutFiles.Add(MovedFrom);
+				}
 
-			// In case of a file Moved/Renamed, consider the rename Origin (where there is now a Redirector file Added)
-			// and add it to the list of files to revert (only if it is not already in) to revert both at once
-			if (!Files.FindByPredicate([&MovedFrom](const FString& File) { return File.Equals(MovedFrom, ESearchCase::IgnoreCase); }))
+				// and delete the Redirector (else the reverted file will collide with it and create a *.private.0 file)
+				IFileManager::Get().Delete(*MovedFrom);
+			}
+			else if (State->WorkspaceState == EWorkspaceState::Added && Operation->ShouldDeleteNewFiles())
 			{
-				Files.Add(MovedFrom);
+				IFileManager::Get().Delete(*File);
 			}
-
-			// and delete the Redirector (else the reverted file will collide with it and create a *.private.0 file)
-			IFileManager::Get().Delete(*MovedFrom);
 		}
+	}
 
-		if (State->WorkspaceState == EWorkspaceState::Added && Operation->ShouldDeleteNewFiles())
+	InCommand.bCommandSuccessful = true;
+
+	if (LocallyChangedFiles.Num() > 0)
+	{
+		if (!GetProvider().IsPartialWorkspace())
 		{
-			IFileManager::Get().Delete(*File);
+			InCommand.bCommandSuccessful &= PlasticSourceControlUtils::RunCommand(TEXT("undochange"), TArray<FString>(), LocallyChangedFiles, InCommand.InfoMessages, InCommand.ErrorMessages);
+		}
+		else
+		{
+			// partial undochange doesn't exist in partial mode
+			InCommand.bCommandSuccessful &= PlasticSourceControlUtils::RunCommand(TEXT("partial undo"), TArray<FString>(), LocallyChangedFiles, InCommand.InfoMessages, InCommand.ErrorMessages);
 		}
 	}
 
-	InCommand.bCommandSuccessful = true;
-
-	if (Files.Num() > 0)
+	if (CheckedOutFiles.Num() > 0)
 	{
+		TArray<FString> Parameters;
+		if (bIsSoftRevert)
+		{
+			Parameters.Add(TEXT("--keepchanges"));
+		}
+
 		// revert the checkout and any changes of the given file in workspace
-		// Detect special case for a partial checkout (CS:-1 in Gluon mode)!
-		if (-1 != InCommand.ChangesetNumber)
+		if (!GetProvider().IsPartialWorkspace())
 		{
-			InCommand.bCommandSuccessful &= PlasticSourceControlUtils::RunCommand(TEXT("undo"), TArray<FString>(), Files, InCommand.InfoMessages, InCommand.ErrorMessages);
+			InCommand.bCommandSuccessful &= PlasticSourceControlUtils::RunCommand(TEXT("undocheckout"), Parameters, CheckedOutFiles, InCommand.InfoMessages, InCommand.ErrorMessages);
 		}
 		else
 		{
-			InCommand.bCommandSuccessful &= PlasticSourceControlUtils::RunCommand(TEXT("partial undo"), TArray<FString>(), Files, InCommand.InfoMessages, InCommand.ErrorMessages);
+			InCommand.bCommandSuccessful &= PlasticSourceControlUtils::RunCommand(TEXT("partial undocheckout"), Parameters, CheckedOutFiles, InCommand.InfoMessages, InCommand.ErrorMessages);
 		}
 	}
 
-	// update the status of our files
-	PlasticSourceControlUtils::RunUpdateStatus(Files, false, InCommand.ErrorMessages, States, InCommand.ChangesetNumber, InCommand.BranchName);
+	// NOTE: optim, in UE4 there was no need to update the status of our files since this is done immediately after by the Editor, except now that we are using changelists
+	// update the status of our files: need to check for local changes in case of a SoftRevert
+	const PlasticSourceControlUtils::EStatusSearchType SearchType = bIsSoftRevert ? PlasticSourceControlUtils::EStatusSearchType::All : PlasticSourceControlUtils::EStatusSearchType::ControlledOnly;
+	PlasticSourceControlUtils::RunUpdateStatus(Files, SearchType, false, InCommand.ErrorMessages, States, InCommand.ChangesetNumber, InCommand.BranchName);
 
 	return InCommand.bCommandSuccessful;
 }
@@ -655,7 +770,7 @@ bool FPlasticRevertUnchangedWorker::Execute(FPlasticSourceControlCommand& InComm
 	{
 		Files.Add(FPaths::ConvertRelativePathToFull(FPaths::ProjectContentDir()));
 	}
-	PlasticSourceControlUtils::RunUpdateStatus(Files, false, InCommand.ErrorMessages, States, InCommand.ChangesetNumber, InCommand.BranchName);
+	PlasticSourceControlUtils::RunUpdateStatus(Files, PlasticSourceControlUtils::EStatusSearchType::ControlledOnly, false, InCommand.ErrorMessages, States, InCommand.ChangesetNumber, InCommand.BranchName);
 
 	return InCommand.bCommandSuccessful;
 }
@@ -667,7 +782,7 @@ bool FPlasticRevertUnchangedWorker::UpdateStates()
 	// Update affected changelists if any
 	for (const FPlasticSourceControlState& NewState : States)
 	{
-		if (!NewState.IsModified())
+		if (!NewState.IsCheckedOutImplementation())
 		{
 			TSharedRef<FPlasticSourceControlState, ESPMode::ThreadSafe> State = GetProvider().GetStateInternal(NewState.GetFilename());
 			if (State->Changelist.IsInitialized())
@@ -701,11 +816,11 @@ bool FPlasticRevertAllWorker::Execute(FPlasticSourceControlCommand& InCommand)
 		TArray<FPlasticSourceControlState> TempStates;
 		TArray<FString> ContentDir;
 		ContentDir.Add(FPaths::ConvertRelativePathToFull(FPaths::ProjectContentDir()));
-		PlasticSourceControlUtils::RunUpdateStatus(ContentDir, false, InCommand.ErrorMessages, TempStates, InCommand.ChangesetNumber, InCommand.BranchName);
+		PlasticSourceControlUtils::RunUpdateStatus(ContentDir, PlasticSourceControlUtils::EStatusSearchType::All, false, InCommand.ErrorMessages, TempStates, InCommand.ChangesetNumber, InCommand.BranchName);
 
 		for (auto& State : TempStates)
 		{
-			if (State.IsModified())
+			if (State.CanRevert())
 			{
 				if (State.WorkspaceState == EWorkspaceState::Added && Operation->ShouldDeleteNewFiles())
 				{
@@ -732,8 +847,7 @@ bool FPlasticRevertAllWorker::Execute(FPlasticSourceControlCommand& InCommand)
 	TArray<FString> Parameters;
 	Parameters.Add(TEXT("--all"));
 	// revert the checkout of all files recursively
-	// Detect special case for a partial checkout (CS:-1 in Gluon mode)!
-	if (-1 != InCommand.ChangesetNumber)
+	if (!GetProvider().IsPartialWorkspace())
 	{
 		InCommand.bCommandSuccessful = PlasticSourceControlUtils::RunCommand(TEXT("undocheckout"), Parameters, TArray<FString>(), Results, InCommand.ErrorMessages);
 	}
@@ -749,7 +863,7 @@ bool FPlasticRevertAllWorker::Execute(FPlasticSourceControlCommand& InCommand)
 	// now update the status of the updated files
 	if (Operation->UpdatedFiles.Num())
 	{
-		PlasticSourceControlUtils::RunUpdateStatus(Operation->UpdatedFiles, false, InCommand.ErrorMessages, States, InCommand.ChangesetNumber, InCommand.BranchName);
+		PlasticSourceControlUtils::RunUpdateStatus(Operation->UpdatedFiles, PlasticSourceControlUtils::EStatusSearchType::ControlledOnly, false, InCommand.ErrorMessages, States, InCommand.ChangesetNumber, InCommand.BranchName);
 	}
 
 	return InCommand.bCommandSuccessful;
@@ -762,7 +876,8 @@ bool FPlasticRevertAllWorker::UpdateStates()
 	// Update affected changelists if any
 	for (const FPlasticSourceControlState& NewState : States)
 	{
-		if (!NewState.IsModified())
+		// TODO: also detect files that were added and are now private! Should be removed as well from their changelist
+		if (!NewState.IsCheckedOutImplementation())
 		{
 			TSharedRef<FPlasticSourceControlState, ESPMode::ThreadSafe> State = GetProvider().GetStateInternal(NewState.GetFilename());
 			if (State->Changelist.IsInitialized())
@@ -804,6 +919,12 @@ bool FPlasticMakeWorkspaceWorker::Execute(FPlasticSourceControlCommand& InComman
 		Parameters.Add(FString::Printf(TEXT("--repository=rep:%s@repserver:%s"), *Operation->RepositoryName, *Operation->ServerUrl));
 		InCommand.bCommandSuccessful = PlasticSourceControlUtils::RunCommand(TEXT("makeworkspace"), Parameters, TArray<FString>(), InCommand.InfoMessages, InCommand.ErrorMessages);
 	}
+	if (Operation->bPartialWorkspace)
+	{
+		TArray<FString> Parameters;
+		Parameters.Add(TEXT("update"));
+		InCommand.bCommandSuccessful = PlasticSourceControlUtils::RunCommand(TEXT("partial"), Parameters, TArray<FString>(), InCommand.InfoMessages, InCommand.ErrorMessages);
+	}
 
 	return InCommand.bCommandSuccessful;
 }
@@ -813,6 +934,241 @@ bool FPlasticMakeWorkspaceWorker::UpdateStates()
 	return false;
 }
 
+FName FPlasticSwitchToPartialWorkspaceWorker::GetName() const
+{
+	return "SwitchToPartialWorkspace";
+}
+
+bool FPlasticSwitchToPartialWorkspaceWorker::Execute(FPlasticSourceControlCommand& InCommand)
+{
+	TRACE_CPUPROFILER_EVENT_SCOPE(FPlasticSwitchToPartialWorkspaceWorker::Execute);
+
+	check(InCommand.Operation->GetName() == GetName());
+
+	{
+		TArray<FString> Parameters;
+		Parameters.Add(TEXT("update"));
+		InCommand.bCommandSuccessful = PlasticSourceControlUtils::RunCommand(TEXT("partial"), Parameters, TArray<FString>(), InCommand.InfoMessages, InCommand.ErrorMessages);
+	}
+
+	// Update the workspace to set the changeset number to -1 if all went well
+	{
+		TArray<FString> ProjectFiles;
+		ProjectFiles.Add(FPaths::ConvertRelativePathToFull(FPaths::GetProjectFilePath()));
+		InCommand.bCommandSuccessful = PlasticSourceControlUtils::RunUpdateStatus(ProjectFiles, PlasticSourceControlUtils::EStatusSearchType::ControlledOnly, false, InCommand.ErrorMessages, States, InCommand.ChangesetNumber, InCommand.BranchName);
+	}
+
+	return InCommand.bCommandSuccessful;
+}
+
+bool FPlasticSwitchToPartialWorkspaceWorker::UpdateStates()
+{
+	return PlasticSourceControlUtils::UpdateCachedStates(MoveTemp(States));
+}
+
+FName FPlasticUnlockWorker::GetName() const
+{
+	return "Unlock";
+}
+
+bool FPlasticUnlockWorker::Execute(FPlasticSourceControlCommand& InCommand)
+{
+	TRACE_CPUPROFILER_EVENT_SCOPE(FPlasticUnlockWorker::Execute);
+
+	check(InCommand.Operation->GetName() == GetName());
+	TSharedRef<FPlasticUnlock, ESPMode::ThreadSafe> Operation = StaticCastSharedRef<FPlasticUnlock>(InCommand.Operation);
+
+	{
+		// retrieve the itemid of assets to unlock
+		FString ItemIds;
+		for (const FString& File : InCommand.Files)
+		{
+			const auto State = GetProvider().GetStateInternal(File);
+			if (State->LockedId != ISourceControlState::INVALID_REVISION)
+			{
+				ItemIds += FString::Printf(TEXT("itemid:%d "), State->LockedId);
+			}
+		}
+
+		TArray<FString> Parameters;
+		Parameters.Add(TEXT("unlock"));
+		if (Operation->bRemove)
+			Parameters.Add(TEXT("--remove"));
+		Parameters.Add(ItemIds);
+		InCommand.bCommandSuccessful = PlasticSourceControlUtils::RunCommand(TEXT("lock"), Parameters, TArray<FString>(), InCommand.InfoMessages, InCommand.ErrorMessages);
+	}
+
+	{
+		InCommand.bCommandSuccessful = PlasticSourceControlUtils::RunUpdateStatus(InCommand.Files, PlasticSourceControlUtils::EStatusSearchType::ControlledOnly, false, InCommand.ErrorMessages, States, InCommand.ChangesetNumber, InCommand.BranchName);
+	}
+
+	return InCommand.bCommandSuccessful;
+}
+
+bool FPlasticUnlockWorker::UpdateStates()
+{
+	return PlasticSourceControlUtils::UpdateCachedStates(MoveTemp(States));
+}
+
+FName FPlasticGetBranchesWorker::GetName() const
+{
+	return "GetBranches";
+}
+
+bool FPlasticGetBranchesWorker::Execute(FPlasticSourceControlCommand& InCommand)
+{
+	TRACE_CPUPROFILER_EVENT_SCOPE(FPlasticGetBranchesWorker::Execute);
+
+	check(InCommand.Operation->GetName() == GetName());
+	TSharedRef<FPlasticGetBranches, ESPMode::ThreadSafe> Operation = StaticCastSharedRef<FPlasticGetBranches>(InCommand.Operation);
+
+	{
+		const FDateTime& FromDate = Operation->FromDate;
+		InCommand.bCommandSuccessful = PlasticSourceControlUtils::RunGetBranches(FromDate, Operation->Branches, InCommand.ErrorMessages);
+	}
+
+	{
+		FString RepositoryName, ServerUrl;
+		InCommand.bCommandSuccessful &= PlasticSourceControlUtils::GetWorkspaceInfo(CurrentBranchName, RepositoryName, ServerUrl, InCommand.ErrorMessages);
+	}
+
+	return InCommand.bCommandSuccessful;
+}
+
+bool FPlasticGetBranchesWorker::UpdateStates()
+{
+	GetProvider().SetBranchName(CurrentBranchName);
+
+	return false;
+}
+
+
+FName FPlasticSwitchToBranchWorker::GetName() const
+{
+	return "Switch";
+}
+
+bool FPlasticSwitchToBranchWorker::Execute(FPlasticSourceControlCommand& InCommand)
+{
+	TRACE_CPUPROFILER_EVENT_SCOPE(FPlasticSwitchToBranchWorker::Execute);
+
+	check(InCommand.Operation->GetName() == GetName());
+	TSharedRef<FPlasticSwitchToBranch, ESPMode::ThreadSafe> Operation = StaticCastSharedRef<FPlasticSwitchToBranch>(InCommand.Operation);
+
+	InCommand.bCommandSuccessful = PlasticSourceControlUtils::RunSwitchToBranch(Operation->BranchName, Operation->UpdatedFiles, InCommand.ErrorMessages);
+
+	// now update the status of the updated files
+	if (InCommand.bCommandSuccessful && Operation->UpdatedFiles.Num())
+	{
+		// the current branch is used to asses the status of Retained Locks
+		GetProvider().SetBranchName(Operation->BranchName);
+		InCommand.bCommandSuccessful = PlasticSourceControlUtils::RunUpdateStatus(Operation->UpdatedFiles, PlasticSourceControlUtils::EStatusSearchType::ControlledOnly, false, InCommand.ErrorMessages, States, InCommand.ChangesetNumber, InCommand.BranchName);
+	}
+
+	return InCommand.bCommandSuccessful;
+}
+
+bool FPlasticSwitchToBranchWorker::UpdateStates()
+{
+	return PlasticSourceControlUtils::UpdateCachedStates(MoveTemp(States));
+}
+
+
+FName FPlasticMergeBranchWorker::GetName() const
+{
+	return "Merge";
+}
+
+bool FPlasticMergeBranchWorker::Execute(FPlasticSourceControlCommand& InCommand)
+{
+	TRACE_CPUPROFILER_EVENT_SCOPE(FPlasticMergeBranchWorker::Execute);
+
+	check(InCommand.Operation->GetName() == GetName());
+	TSharedRef<FPlasticMergeBranch, ESPMode::ThreadSafe> Operation = StaticCastSharedRef<FPlasticMergeBranch>(InCommand.Operation);
+
+	InCommand.bCommandSuccessful = PlasticSourceControlUtils::RunMergeBranch(Operation->BranchName, Operation->UpdatedFiles, InCommand.ErrorMessages);
+
+	// now update the status of the updated files
+	if (Operation->UpdatedFiles.Num())
+	{
+		InCommand.bCommandSuccessful = PlasticSourceControlUtils::RunUpdateStatus(Operation->UpdatedFiles, PlasticSourceControlUtils::EStatusSearchType::ControlledOnly, false, InCommand.ErrorMessages, States, InCommand.ChangesetNumber, InCommand.BranchName);
+	}
+
+	return InCommand.bCommandSuccessful;
+}
+
+bool FPlasticMergeBranchWorker::UpdateStates()
+{
+	// Update the Default changelist.
+	UpdateChangelistState(GetProvider(), FPlasticSourceControlChangelist::DefaultChangelist, States);
+
+	return PlasticSourceControlUtils::UpdateCachedStates(MoveTemp(States));
+}
+
+
+FName FPlasticCreateBranchWorker::GetName() const
+{
+	return "CreateBranch";
+}
+
+bool FPlasticCreateBranchWorker::Execute(FPlasticSourceControlCommand& InCommand)
+{
+	TRACE_CPUPROFILER_EVENT_SCOPE(FPlasticCreateBranchWorker::Execute);
+
+	check(InCommand.Operation->GetName() == GetName());
+	TSharedRef<FPlasticCreateBranch, ESPMode::ThreadSafe> Operation = StaticCastSharedRef<FPlasticCreateBranch>(InCommand.Operation);
+
+	return PlasticSourceControlUtils::RunCreateBranch(Operation->BranchName, Operation->Comment, InCommand.ErrorMessages);
+}
+
+bool FPlasticCreateBranchWorker::UpdateStates()
+{
+	return false;
+}
+
+
+FName FPlasticRenameBranchWorker::GetName() const
+{
+	return "RenameBranch";
+}
+
+bool FPlasticRenameBranchWorker::Execute(FPlasticSourceControlCommand& InCommand)
+{
+	TRACE_CPUPROFILER_EVENT_SCOPE(FPlasticRenameBranchWorker::Execute);
+
+	check(InCommand.Operation->GetName() == GetName());
+	TSharedRef<FPlasticRenameBranch, ESPMode::ThreadSafe> Operation = StaticCastSharedRef<FPlasticRenameBranch>(InCommand.Operation);
+
+	return PlasticSourceControlUtils::RunRenameBranch(Operation->OldName, Operation->NewName, InCommand.ErrorMessages);
+}
+
+bool FPlasticRenameBranchWorker::UpdateStates()
+{
+	return false;
+}
+
+
+FName FPlasticDeleteBranchesWorker::GetName() const
+{
+	return "DeleteBranches";
+}
+
+bool FPlasticDeleteBranchesWorker::Execute(FPlasticSourceControlCommand& InCommand)
+{
+	TRACE_CPUPROFILER_EVENT_SCOPE(FPlasticDeleteBranchesWorker::Execute);
+
+	check(InCommand.Operation->GetName() == GetName());
+	TSharedRef<FPlasticDeleteBranches, ESPMode::ThreadSafe> Operation = StaticCastSharedRef<FPlasticDeleteBranches>(InCommand.Operation);
+
+	return PlasticSourceControlUtils::RunDeleteBranches(Operation->BranchNames, InCommand.ErrorMessages);
+}
+
+bool FPlasticDeleteBranchesWorker::UpdateStates()
+{
+	return false;
+}
+
+
 FName FPlasticUpdateStatusWorker::GetName() const
 {
 	return "UpdateStatus";
@@ -826,28 +1182,21 @@ bool FPlasticUpdateStatusWorker::Execute(FPlasticSourceControlCommand& InCommand
 	TSharedRef<FUpdateStatus, ESPMode::ThreadSafe> Operation = StaticCastSharedRef<FUpdateStatus>(InCommand.Operation);
 
 	// Note: ShouldCheckAllFiles is never set to true (SetCheckingAllFiles)
-	UE_LOG(LogSourceControl, Log, TEXT("status (of %d files, ShouldUpdateHistory=%d, ShouldGetOpenedOnly=%d, ShouldUpdateModifiedState=%d)"),
+	UE_LOG(LogSourceControl, Log, TEXT("status of %d items (ShouldUpdateHistory=%d, ShouldGetOpenedOnly=%d, ShouldUpdateModifiedState=%d)"),
 		InCommand.Files.Num(), Operation->ShouldUpdateHistory(), Operation->ShouldGetOpenedOnly(), Operation->ShouldUpdateModifiedState());
 
 	const TArray<FString> Files = GetFilesFromCommand(GetProvider(), InCommand);
 
 	if (Files.Num() > 0)
 	{
-		InCommand.bCommandSuccessful = PlasticSourceControlUtils::RunUpdateStatus(Files, Operation->ShouldUpdateHistory(), InCommand.ErrorMessages, States, InCommand.ChangesetNumber, InCommand.BranchName);
+		InCommand.bCommandSuccessful = PlasticSourceControlUtils::RunUpdateStatus(Files, PlasticSourceControlUtils::EStatusSearchType::All, Operation->ShouldUpdateHistory(), InCommand.ErrorMessages, States, InCommand.ChangesetNumber, InCommand.BranchName);
 		// Remove all "is not in a workspace" error and convert the result to "success" if there are no other errors
 		PlasticSourceControlUtils::RemoveRedundantErrors(InCommand, TEXT("is not in a workspace."));
 		if (!InCommand.bCommandSuccessful)
 		{
-			TArray<FString> ErrorMessages;
-			TArray<FString> Parameters;
-			FString BranchName, RepositoryName, ServerUrl;
-			if (PlasticSourceControlUtils::GetWorkspaceInfo(BranchName, RepositoryName, ServerUrl, ErrorMessages))
-			{
-				Parameters.Add(FString::Printf(TEXT("--server=%s"), *ServerUrl));
-			}
-			UE_LOG(LogSourceControl, Error, TEXT("FPlasticUpdateStatusWorker(ErrorMessages.Num()=%d) => checkconnection"), InCommand.ErrorMessages.Num());
-			// In case of error, execute a 'checkconnection' command to check the connectivity of the server.
-			InCommand.bConnectionDropped = !PlasticSourceControlUtils::RunCommand(TEXT("checkconnection"), Parameters, TArray<FString>(), InCommand.InfoMessages, InCommand.ErrorMessages);
+			// In case of error, execute a 'checkconnection' command to check the connection to the server.
+			UE_LOG(LogSourceControl, Warning, TEXT("Error on 'status', execute a 'checkconnection' to test the connection to the server"));
+			InCommand.bConnectionDropped = !PlasticSourceControlUtils::RunCheckConnection(InCommand.BranchName, InCommand.RepositoryName, InCommand.ServerUrl, InCommand.InfoMessages, InCommand.ErrorMessages);
 			return false;
 		}
 
@@ -874,18 +1223,18 @@ bool FPlasticUpdateStatusWorker::Execute(FPlasticSourceControlCommand& InCommand
 	{
 		TArray<FString> ProjectDirs;
 		ProjectDirs.Add(FPaths::ConvertRelativePathToFull(FPaths::ProjectContentDir()));
-		InCommand.bCommandSuccessful = PlasticSourceControlUtils::RunUpdateStatus(ProjectDirs, Operation->ShouldUpdateHistory(), InCommand.ErrorMessages, States, InCommand.ChangesetNumber, InCommand.BranchName);
+		InCommand.bCommandSuccessful = PlasticSourceControlUtils::RunUpdateStatus(ProjectDirs, PlasticSourceControlUtils::EStatusSearchType::All, Operation->ShouldUpdateHistory(), InCommand.ErrorMessages, States, InCommand.ChangesetNumber, InCommand.BranchName);
 	}
 	else
 	{
-		// TODO: workaround for the case of submitting a changelist, calling UpdateStatus with no files nor the changelist.
+		// Note: workaround for the case of submitting a changelist, calling UpdateStatus with no files nor the changelist.
 		// No consequences, and no way to fix it, so let's not show an error.
 		InCommand.bCommandSuccessful = true;
 	}
 
 	// TODO: re-evaluate how to optimize this heavy operation using some of these hints flags
 	// - ShouldGetOpenedOnly hint would be to call for all a whole workspace status update
-	// - ShouldUpdateModifiedState hint not used as the above normal Plastic status has already told us this information (like Git and Mercurial)
+	// - ShouldUpdateModifiedState hint not used as the above normal status has already told us this information (like Git and Mercurial)
 
 	return InCommand.bCommandSuccessful;
 }
@@ -897,7 +1246,7 @@ bool FPlasticUpdateStatusWorker::UpdateStates()
 	// Update affected changelists if any (in case of a file reverted outside of the Unreal Editor)
 	for (const FPlasticSourceControlState& NewState : States)
 	{
-		if (!NewState.IsModified())
+		if (!NewState.IsCheckedOutImplementation())
 		{
 			TSharedRef<FPlasticSourceControlState, ESPMode::ThreadSafe> State = GetProvider().GetStateInternal(NewState.GetFilename());
 			if (State->Changelist.IsInitialized())
@@ -936,7 +1285,7 @@ bool FPlasticCopyWorker::Execute(FPlasticSourceControlCommand& InCommand)
 		if (bIsMoveOperation)
 		{
 			UE_LOG(LogSourceControl, Log, TEXT("Moving %s to %s..."), *Origin, *Destination);
-			// In case of rename, we have to undo what the Editor (created a redirector and added the dest asset), and then redo it with Plastic SCM
+			// In case of rename, we have to undo what the Editor (created a redirector and added the dest asset), and then redo it with Unity Version Control
 			// - revert the 'cm add' that was applied to the destination by the Editor
 			{
 				TArray<FString> DestinationFiles;
@@ -951,8 +1300,7 @@ bool FPlasticCopyWorker::Execute(FPlasticSourceControlCommand& InCommand)
 				TArray<FString> Files;
 				Files.Add(Origin);
 				Files.Add(Destination);
-				// Detect special case for a partial checkout (CS:-1 in Gluon mode)!
-				if (-1 != InCommand.ChangesetNumber)
+				if (!GetProvider().IsPartialWorkspace())
 				{
 					InCommand.bCommandSuccessful = PlasticSourceControlUtils::RunCommand(TEXT("move"), Parameters, Files, InCommand.InfoMessages, InCommand.ErrorMessages);
 				}
@@ -966,8 +1314,7 @@ bool FPlasticCopyWorker::Execute(FPlasticSourceControlCommand& InCommand)
 			{
 				TArray<FString> Files;
 				Files.Add(Origin);
-				// Detect special case for a partial checkout (CS:-1 in Gluon mode)!
-				if (-1 != InCommand.ChangesetNumber)
+				if (!GetProvider().IsPartialWorkspace())
 				{
 					InCommand.bCommandSuccessful = PlasticSourceControlUtils::RunCommand(TEXT("add"), TArray<FString>(), Files, InCommand.InfoMessages, InCommand.ErrorMessages);
 				}
@@ -983,11 +1330,11 @@ bool FPlasticCopyWorker::Execute(FPlasticSourceControlCommand& InCommand)
 			InCommand.bCommandSuccessful = true;
 		}
 
-		// now update the status of our files:
+		// now update the status of our files
 		TArray<FString> BothFiles;
 		BothFiles.Add(Origin);
 		BothFiles.Add(Destination);
-		PlasticSourceControlUtils::RunUpdateStatus(BothFiles, false, InCommand.ErrorMessages, States, InCommand.ChangesetNumber, InCommand.BranchName);
+		PlasticSourceControlUtils::RunUpdateStatus(BothFiles, PlasticSourceControlUtils::EStatusSearchType::ControlledOnly, false, InCommand.ErrorMessages, States, InCommand.ChangesetNumber, InCommand.BranchName);
 	}
 	else
 	{
@@ -999,7 +1346,7 @@ bool FPlasticCopyWorker::Execute(FPlasticSourceControlCommand& InCommand)
 
 bool FPlasticCopyWorker::UpdateStates()
 {
-	TRACE_CPUPROFILER_EVENT_SCOPE(FPlasticCopyWorkers::UpdateStates);
+	TRACE_CPUPROFILER_EVENT_SCOPE(FPlasticCopyWorker::UpdateStates);
 
 	return PlasticSourceControlUtils::UpdateCachedStates(MoveTemp(States));
 }
@@ -1019,7 +1366,7 @@ bool FPlasticSyncWorker::Execute(FPlasticSourceControlCommand& InCommand)
 	// now update the status of the updated files
 	if (UpdatedFiles.Num())
 	{
-		PlasticSourceControlUtils::RunUpdateStatus(UpdatedFiles, false, InCommand.ErrorMessages, States, InCommand.ChangesetNumber, InCommand.BranchName);
+		PlasticSourceControlUtils::RunUpdateStatus(UpdatedFiles, PlasticSourceControlUtils::EStatusSearchType::ControlledOnly, false, InCommand.ErrorMessages, States, InCommand.ChangesetNumber, InCommand.BranchName);
 	}
 
 	if ((InCommand.Operation->GetName() == FName("SyncAll")))
@@ -1067,13 +1414,12 @@ bool FPlasticResolveWorker::Execute(FPlasticSourceControlCommand& InCommand)
 		OneFile.Add(State->PendingResolveInfo.BaseFile);
 
 		UE_LOG(LogSourceControl, Log, TEXT("resolve %s"), *State->PendingResolveInfo.BaseFile);
-
 		// Mark the conflicted file as resolved
 		InCommand.bCommandSuccessful = PlasticSourceControlUtils::RunCommand(TEXT("merge"), Parameters, OneFile, InCommand.InfoMessages, InCommand.ErrorMessages);
 	}
 
 	// now update the status of our files
-	PlasticSourceControlUtils::RunUpdateStatus(InCommand.Files, false, InCommand.ErrorMessages, States, InCommand.ChangesetNumber, InCommand.BranchName);
+	PlasticSourceControlUtils::RunUpdateStatus(InCommand.Files, PlasticSourceControlUtils::EStatusSearchType::ControlledOnly, false, InCommand.ErrorMessages, States, InCommand.ChangesetNumber, InCommand.BranchName);
 
 	return InCommand.bCommandSuccessful;
 }
@@ -1085,6 +1431,7 @@ bool FPlasticResolveWorker::UpdateStates()
 	return PlasticSourceControlUtils::UpdateCachedStates(MoveTemp(States));
 }
 
+
 FName FPlasticGetPendingChangelistsWorker::GetName() const
 {
 	return "UpdateChangelistsStatus";
@@ -1131,7 +1478,7 @@ bool FPlasticGetPendingChangelistsWorker::UpdateStates()
 	const FDateTime Now = FDateTime::Now();
 
 	// first update cached state from 'changes' call
-	for (int StatusIndex = 0; StatusIndex < OutChangelistsStates.Num(); StatusIndex++)
+	for (int32 StatusIndex = 0; StatusIndex < OutChangelistsStates.Num(); StatusIndex++)
 	{
 		const FPlasticSourceControlChangelistState& CLStatus = OutChangelistsStates[StatusIndex];
 		TSharedRef<FPlasticSourceControlChangelistState, ESPMode::ThreadSafe> ChangelistState = GetProvider().GetStateInternal(CLStatus.Changelist);
@@ -1140,7 +1487,7 @@ bool FPlasticGetPendingChangelistsWorker::UpdateStates()
 		ChangelistState->TimeStamp = Now;
 		bUpdated = true;
 
-		// Update files states for files in the changelist
+		// Update files in the changelist
 		bool bUpdateFilesStates = (OutCLFilesStates.Num() == OutChangelistsStates.Num());
 		if (bUpdateFilesStates)
 		{
@@ -1216,9 +1563,9 @@ FPlasticSourceControlChangelist CreatePendingChangelist(FPlasticSourceControlPro
 	{
 		Parameters.Add(TEXT("create"));
 		const FScopedTempFile ChangelistNameFile(NewChangelist.GetName());
-		Parameters.Add(FString::Printf(TEXT("--namefile=\"%s\""), *FPaths::ConvertRelativePathToFull(ChangelistNameFile.GetFilename())));
+		Parameters.Add(FString::Printf(TEXT("--namefile=\"%s\""), *ChangelistNameFile.GetFilename()));
 		const FScopedTempFile ChangelistDescriptionFile(InDescription);
-		Parameters.Add(FString::Printf(TEXT("--descriptionfile=\"%s\""), *FPaths::ConvertRelativePathToFull(ChangelistDescriptionFile.GetFilename())));
+		Parameters.Add(FString::Printf(TEXT("--descriptionfile=\"%s\""), *ChangelistDescriptionFile.GetFilename()));
 		Parameters.Add(TEXT("--persistent")); // Create a persistent changelist to stay close to Perforce behavior
 		UE_LOG(LogSourceControl, Verbose, TEXT("CreatePendingChangelist(%s):\n\"%s\""), *NewChangelist.GetName(), *InDescription);
 		bCommandSuccessful = PlasticSourceControlUtils::RunCommand(TEXT("changelist"), Parameters, TArray<FString>(), InInfoMessages, InErrorMessages);
@@ -1245,10 +1592,10 @@ bool EditChangelistDescription(const FPlasticSourceControlProvider& PlasticSourc
 	else
 	{
 		const FScopedTempFile ChangelistNameFile(InChangelist.GetName());
-		Parameters.Add(FString::Printf(TEXT("--namefile=\"%s\""), *FPaths::ConvertRelativePathToFull(ChangelistNameFile.GetFilename())));
+		Parameters.Add(FString::Printf(TEXT("--namefile=\"%s\""), *ChangelistNameFile.GetFilename()));
 		Parameters.Add(TEXT("description"));
 		const FScopedTempFile ChangelistDescriptionFile(InDescription);
-		Parameters.Add(FString::Printf(TEXT("--descriptionfile=\"%s\""), *FPaths::ConvertRelativePathToFull(ChangelistDescriptionFile.GetFilename())));
+		Parameters.Add(FString::Printf(TEXT("--descriptionfile=\"%s\""), *ChangelistDescriptionFile.GetFilename()));
 		UE_LOG(LogSourceControl, Verbose, TEXT("EditChangelistDescription(%s\n%s)"), *InChangelist.GetName(), *InDescription);
 		return PlasticSourceControlUtils::RunCommand(TEXT("changelist"), Parameters, TArray<FString>(), InInfoMessages, InErrorMessages);
 	}
@@ -1268,7 +1615,7 @@ bool MoveFilesToChangelist(const FPlasticSourceControlProvider& PlasticSourceCon
 		else
 		{
 			const FScopedTempFile ChangelistNameFile(InChangelist.GetName());
-			Parameters.Add(FString::Printf(TEXT("--namefile=\"%s\""), *FPaths::ConvertRelativePathToFull(ChangelistNameFile.GetFilename())));
+			Parameters.Add(FString::Printf(TEXT("--namefile=\"%s\""), *ChangelistNameFile.GetFilename()));
 			Parameters.Add(TEXT("add"));
 			UE_LOG(LogSourceControl, Verbose, TEXT("MoveFilesToChangelist(%s)"), *InChangelist.GetName());
 			return PlasticSourceControlUtils::RunCommand(TEXT("changelist"), Parameters, InFiles, OutResults, OutErrorMessages);
@@ -1334,6 +1681,8 @@ bool FPlasticNewChangelistWorker::Execute(class FPlasticSourceControlCommand& In
 
 bool FPlasticNewChangelistWorker::UpdateStates()
 {
+	TRACE_CPUPROFILER_EVENT_SCOPE(FPlasticNewChangelistWorker::UpdateStates);
+
 	if (NewChangelist.IsInitialized())
 	{
 		const FDateTime Now = FDateTime::Now();
@@ -1567,7 +1916,7 @@ bool CreateShelve(const FString& InChangelistName, const FString& InChangelistDe
 	const FString ShelveDescription = FString::Printf(TEXT("Changelist%s: %s"), *InChangelistName, *InChangelistDescription);
 	const FScopedTempFile CommentsFile(ShelveDescription);
 	Parameters.Add(TEXT("create"));
-	Parameters.Add(FString::Printf(TEXT("-commentsfile=\"%s\""), *FPaths::ConvertRelativePathToFull(CommentsFile.GetFilename())));
+	Parameters.Add(FString::Printf(TEXT("-commentsfile=\"%s\""), *CommentsFile.GetFilename()));
 	const bool bCommandSuccessful = PlasticSourceControlUtils::RunCommand(TEXT("shelveset"), Parameters, InFilesToShelve, Results, OutErrorMessages);
 	if (bCommandSuccessful && Results.Num() > 0)
 	{
@@ -1675,9 +2024,35 @@ bool FPlasticShelveWorker::Execute(FPlasticSourceControlCommand& InCommand)
 
 	if (InCommand.bCommandSuccessful)
 	{
+		// Remove unmodified files from the list of files to shelve
+		int32 i = 0;
+		while (i < FilesToShelve.Num())
+		{
+			FString& FileToShelve = FilesToShelve[i];
+			TSharedRef<FPlasticSourceControlState, ESPMode::ThreadSafe> FileState = GetProvider().GetStateInternal(FileToShelve);
+			if (FileState->IsModified())
+			{
+				i++;
+			}
+			else
+			{
+				FPaths::MakePathRelativeTo(FileToShelve, *FPaths::ProjectDir());
+				UE_LOG(LogSourceControl, Warning, TEXT("The file /%s is unchanged, it cannot be shelved."), *FileToShelve);
+				FilesToShelve.RemoveAt(i);
+			}
+		}
+
 		ChangelistDescription = *Operation->GetDescription().ToString();
 
-		InCommand.bCommandSuccessful = CreateShelve(Changelist.GetName(), ChangelistDescription, FilesToShelve, ShelveId, InCommand.ErrorMessages);
+		if (FilesToShelve.Num() > 0)
+		{
+			InCommand.bCommandSuccessful = CreateShelve(Changelist.GetName(), ChangelistDescription, FilesToShelve, ShelveId, InCommand.ErrorMessages);
+		}
+		else
+		{
+			UE_LOG(LogSourceControl, Error, TEXT("No file to Shelve"));
+			InCommand.bCommandSuccessful = false;
+		}
 		if (InCommand.bCommandSuccessful)
 		{
 			InChangelistToUpdate = InCommand.Changelist;
@@ -1701,6 +2076,7 @@ bool FPlasticShelveWorker::Execute(FPlasticSourceControlCommand& InCommand)
 				}
 
 				DeleteChangelist(GetProvider(), Changelist, InCommand.InfoMessages, InCommand.ErrorMessages);
+				GetProvider().RemoveChangelistFromCache(Changelist);
 			}
 		}
 	}
@@ -1741,13 +2117,14 @@ bool FPlasticShelveWorker::UpdateStates()
 		}
 
 		DestinationChangelistState->ShelveId = ShelveId;
+		DestinationChangelistState->ShelveDate = FDateTime::Now();
 
 		// And finally, add the shelved files to the changelist state
 		DestinationChangelistState->ShelvedFiles.Reset(ShelvedFiles.Num());
 		for (FString ShelvedFile : ShelvedFiles)
 		{
 			TSharedRef<FPlasticSourceControlState, ESPMode::ThreadSafe> FileState = GetProvider().GetStateInternal(ShelvedFile);
-			PlasticSourceControlUtils::AddShelvedFileToChangelist(DestinationChangelistState.Get(), MoveTemp(ShelvedFile), FileState->WorkspaceState);
+			PlasticSourceControlUtils::AddShelvedFileToChangelist(DestinationChangelistState.Get(), MoveTemp(ShelvedFile), FileState->WorkspaceState, FString(FileState->MovedFrom));
 		}
 
 		return bMovedFiles || ShelvedFiles.Num() > 0;
@@ -1796,7 +2173,7 @@ bool FPlasticUnshelveWorker::Execute(FPlasticSourceControlCommand& InCommand)
 			{
 				// On old version, don't unshelve the files if they are not all selected (since we couldn't apply only a selection of files from a shelve)
 				UE_LOG(LogSourceControl, Error,
-					TEXT("Plastic SCM %s cannot unshelve a selection of files from a shelve. Unshelve them all at once or update to %s or above."),
+					TEXT("Unity Version Control %s cannot unshelve a selection of files from a shelve. Unshelve them all at once or update to %s or above."),
 					*GetProvider().GetPlasticScmVersion().String,
 					*PlasticSourceControlVersions::ShelvesetApplySelection.String
 				);
@@ -1836,7 +2213,7 @@ bool FPlasticUnshelveWorker::Execute(FPlasticSourceControlCommand& InCommand)
 	if (InCommand.bCommandSuccessful)
 	{
 		// now update the status of our files
-		PlasticSourceControlUtils::RunUpdateStatus(InCommand.Files, false, InCommand.ErrorMessages, States, InCommand.ChangesetNumber, InCommand.BranchName);
+		PlasticSourceControlUtils::RunUpdateStatus(InCommand.Files, PlasticSourceControlUtils::EStatusSearchType::ControlledOnly, false, InCommand.ErrorMessages, States, InCommand.ChangesetNumber, InCommand.BranchName);
 
 		ChangelistToUpdate = InCommand.Changelist;
 	}
@@ -1875,7 +2252,7 @@ bool FPlasticDeleteShelveWorker::Execute(FPlasticSourceControlCommand& InCommand
 			if (!InCommand.Files.Contains(ShelveState->GetFilename()))
 			{
 				FString File = ShelveState->GetFilename();
-				
+
 				// Check that all this files are is still in the corresponding changelist, else we won't be able to create the new shelve!
 				if (ChangelistState->Files.ContainsByPredicate([&File](FSourceControlStateRef& State) { return File == State->GetFilename(); }))
 				{
@@ -1883,7 +2260,6 @@ bool FPlasticDeleteShelveWorker::Execute(FPlasticSourceControlCommand& InCommand
 				}
 				else
 				{
-
 					FPaths::MakePathRelativeTo(File, *FPaths::ProjectDir());
 					UE_LOG(LogSourceControl, Error, TEXT("The file /%s is not in the changelist anymore, so the shelve cannot be updated. Unshelve the corresponding change and retry."), *File);
 					InCommand.bCommandSuccessful = false;
@@ -1897,7 +2273,7 @@ bool FPlasticDeleteShelveWorker::Execute(FPlasticSourceControlCommand& InCommand
 		ChangelistToUpdate = InCommand.Changelist;
 		FilesToRemove = InCommand.Files;
 	}
-	
+
 	if (InCommand.bCommandSuccessful && FilesToShelve.Num() > 0)
 	{
 		// Create a new shelve with the other files
@@ -1925,7 +2301,7 @@ bool FPlasticDeleteShelveWorker::UpdateStates()
 
 		if (FilesToRemove.Num() > 0)
 		{
-			// NOTE: for now, Plastic SCM cannot delete a selection of files from a shelve, so FilesToRemove and this specific case aren't really needed (yet)
+			// NOTE: for now, Unity Version Control cannot delete a selection of files from a shelve, so FilesToRemove and this specific case aren't really needed (yet)
 			return ChangelistState->ShelvedFiles.RemoveAll([this](FSourceControlStateRef& State) -> bool
 				{
 					return Algo::AnyOf(FilesToRemove, [&State](const FString& File) {
@@ -1946,4 +2322,132 @@ bool FPlasticDeleteShelveWorker::UpdateStates()
 	}
 }
 
+// Copied from SSourceControlReview.cpp of the ChangelistReview Editor plugin
+namespace ReviewHelpers
+{
+	const FString FileDepotKey = TEXT("depotFile");
+	const FString FileRevisionKey = TEXT("rev");
+	const FString FileActionKey = TEXT("action");
+	const FString TimeKey = TEXT("time");
+	const FString AuthorKey = TEXT("user");
+	const FString DescriptionKey = TEXT("desc");
+	const FString ChangelistStatusKey = TEXT("status");
+	const FString ChangelistPendingStatusKey = TEXT("pending");
+	constexpr int32 RecordIndex = 0;
+} // namespace ReviewHelpers
+
+FName FPlasticGetChangelistDetailsWorker::GetName() const
+{
+	return "GetChangelistDetails";
+}
+
+bool FPlasticGetChangelistDetailsWorker::Execute(FPlasticSourceControlCommand& InCommand)
+{
+	TRACE_CPUPROFILER_EVENT_SCOPE(FPlasticGetChangelistDetailsWorker::Execute);
+
+	TSharedRef<FGetChangelistDetails, ESPMode::ThreadSafe> Operation = StaticCastSharedRef<FGetChangelistDetails>(InCommand.Operation);
+
+	// Note: Changelists are local construct so we have to interpret this as a Shelve Id instead
+	const FString& ShelveId = Operation->GetChangelistNumber();
+	if (ShelveId.IsEmpty())
+	{
+		InCommand.bCommandSuccessful = false;
+		InCommand.ErrorMessages.Add((LOCTEXT("GetChangelistDetailsEmptyId", "GetChangelistDetails failed. Shelve Id is empty.").ToString()));
+		return false;
+	}
+
+	FString Comment;
+	FString Owner;
+	FDateTime Date;
+	TArray<FPlasticSourceControlRevision> BaseRevisions;
+	InCommand.bCommandSuccessful = PlasticSourceControlUtils::RunGetShelve(FCString::Atoi(*ShelveId), Comment, Date, Owner, BaseRevisions, InCommand.ErrorMessages);
+	if (!InCommand.bCommandSuccessful)
+	{
+		InCommand.bCommandSuccessful = false;
+		InCommand.ErrorMessages.Add((LOCTEXT("GetChangelistDetailsInvalidId", "GetChangelistDetails failed. Shelve Id is invalid.").ToString()));
+		return false;
+	}
+
+	UE_LOG(LogSourceControl, Log, TEXT("GetChangelistDetails: %d files in shelve %s"), BaseRevisions.Num(), *ShelveId);
+
+	TMap<FString, FString> Record;
+
+	Record.Add({ ReviewHelpers::ChangelistStatusKey, ReviewHelpers::ChangelistPendingStatusKey });
+	Record.Add({ ReviewHelpers::AuthorKey, Owner });
+	Record.Add({ ReviewHelpers::DescriptionKey, Comment });
+	Record.Add({ ReviewHelpers::TimeKey, LexToString(Date.ToUnixTimestamp()) });
+
+	uint32  RecordFileIndex = 0;
+	for (auto& Revision : BaseRevisions)
+	{
+		// String representation of the current file index
+		FString RecordFileIndexStr = LexToString(RecordFileIndex);
+		// The p4 records is the map a file key starts with "depotFile" and is followed by file index
+		FString RecordFileMapKey = ReviewHelpers::FileDepotKey + RecordFileIndexStr;
+		// The p4 records is the map a revision key starts with "rev" and is followed by file index
+		FString RecordRevisionMapKey = ReviewHelpers::FileRevisionKey + RecordFileIndexStr;
+		// The p4 records is the map a revision key starts with "action" and is followed by file index
+		FString RecordActionMapKey = ReviewHelpers::FileActionKey + RecordFileIndexStr;
+
+		UE_LOG(LogSourceControl, Log, TEXT("GetChangelistDetails: %s baserevid:%d %s"), *Revision.Filename, Revision.RevisionId, *Revision.Action);
+
+		Record.Add({ MoveTemp(RecordFileMapKey), MoveTemp(Revision.Filename) });
+		Record.Add({ MoveTemp(RecordRevisionMapKey), LexToString(Revision.RevisionId) });
+		Record.Add({ MoveTemp(RecordActionMapKey), Revision.Action });
+
+		RecordFileIndex++;
+	}
+
+	TArray<TMap<FString, FString>> ChangelistRecord;
+	ChangelistRecord.Add(MoveTemp(Record));
+	Operation->SetChangelistDetails(MoveTemp(ChangelistRecord));
+
+	return InCommand.bCommandSuccessful;
+}
+
+bool FPlasticGetChangelistDetailsWorker::UpdateStates()
+{
+	return false;
+}
+
+FName FPlasticGetFileWorker::GetName() const
+{
+	return "GetFile";
+}
+
+bool FPlasticGetFileWorker::Execute(FPlasticSourceControlCommand& InCommand)
+{
+	TRACE_CPUPROFILER_EVENT_SCOPE(FPlasticGetFileWorker::Execute);
+
+	TSharedRef<FGetFile, ESPMode::ThreadSafe> Operation = StaticCastSharedRef<FGetFile>(InCommand.Operation);
+
+	const TSharedRef<FPlasticSourceControlRevision, ESPMode::ThreadSafe> SourceControlRevision = MakeShared<FPlasticSourceControlRevision>();
+	SourceControlRevision->Filename = FPaths::ConvertRelativePathToFull(Operation->GetDepotFilePath());
+
+	if (Operation->IsShelve())
+	{
+		SourceControlRevision->ShelveId = FCString::Atoi(*Operation->GetChangelistNumber());
+		UE_LOG(LogSourceControl, Log, TEXT("GetFile(ShelveId:%d)"), SourceControlRevision->ShelveId);
+	}
+	else
+	{
+		SourceControlRevision->RevisionId = FCString::Atoi(*Operation->GetRevisionNumber());
+		UE_LOG(LogSourceControl, Log, TEXT("GetFile(revid:%d)"), SourceControlRevision->RevisionId);
+	}
+
+	FString OutFilename;
+	InCommand.bCommandSuccessful = SourceControlRevision->Get(OutFilename, InCommand.Concurrency);
+	if (InCommand.bCommandSuccessful)
+	{
+		Operation->SetOutPackageFilename(OutFilename);
+	}
+
+	return InCommand.bCommandSuccessful;
+}
+
+bool FPlasticGetFileWorker::UpdateStates()
+{
+	return false;
+}
+
 #undef LOCTEXT_NAMESPACE
diff --git a/Engine/Plugins/Developer/PlasticSourceControl/Source/PlasticSourceControl/Private/PlasticSourceControlOperations.h b/Engine/Plugins/Developer/PlasticSourceControl/Source/PlasticSourceControl/Private/PlasticSourceControlOperations.h
--- a/Engine/Plugins/Developer/PlasticSourceControl/Source/PlasticSourceControl/Private/PlasticSourceControlOperations.h
+++ b/Engine/Plugins/Developer/PlasticSourceControl/Source/PlasticSourceControl/Private/PlasticSourceControlOperations.h
@@ -8,18 +8,21 @@
 #include "IPlasticSourceControlWorker.h"
 #include "PlasticSourceControlRevision.h"
 #include "PlasticSourceControlState.h"
-#include "ISourceControlOperation.h"
+#include "SourceControlOperationBase.h"
 #include "SourceControlOperations.h"
 
 #include "PlasticSourceControlChangelist.h"
 #include "PlasticSourceControlChangelistState.h"
 
 class FPlasticSourceControlProvider;
+typedef TSharedRef<class FPlasticSourceControlBranch, ESPMode::ThreadSafe> FPlasticSourceControlBranchRef;
+
 
 /**
  * Internal operation used to revert checked-out unchanged files
 */
-class FPlasticRevertUnchanged final : public ISourceControlOperation
+// NOTE: added to Engine in Unreal Engine 5 for changelists
+class FPlasticRevertUnchanged final : public FSourceControlOperationBase
 {
 public:
 	// ISourceControlOperation interface
@@ -36,10 +39,7 @@ class FPlasticSyncAll final : public FSync
 {
 public:
 	// ISourceControlOperation interface
-	virtual FName GetName() const override
-	{
-		return "SyncAll";
-	}
+	virtual FName GetName() const override;
 
 	/** List of files updated by the operation */
 	TArray<FString> UpdatedFiles;
@@ -63,9 +63,9 @@ public:
 
 
 /**
-* Internal operation used to initialize a new Workspace and a new Repository
+* Internal operation used to create a new Workspace and a new Repository
 */
-class FPlasticMakeWorkspace final : public ISourceControlOperation
+class FPlasticMakeWorkspace final : public FSourceControlOperationBase
 {
 public:
 	// ISourceControlOperation interface
@@ -76,6 +76,140 @@ public:
 	FString WorkspaceName;
 	FString RepositoryName;
 	FString ServerUrl;
+	bool bPartialWorkspace = false;
+};
+
+
+/**
+ * Internal operation used to switch to a partial workspace
+*/
+class FPlasticSwitchToPartialWorkspace final : public FSourceControlOperationBase
+{
+public:
+	// ISourceControlOperation interface
+	virtual FName GetName() const override;
+
+	virtual FText GetInProgressString() const override;
+};
+
+
+/**
+ * Internal operation used to release or remove Lock(s) on file(s)
+*/
+class FPlasticUnlock final : public FSourceControlOperationBase
+{
+public:
+	// ISourceControlOperation interface
+	virtual FName GetName() const override;
+
+	virtual FText GetInProgressString() const override;
+
+	// Release the Lock(s), and optionally remove (delete) them completely
+	bool bRemove = false;
+};
+
+
+/**
+ * Internal operation to list branches, aka "cm find branch"
+*/
+class FPlasticGetBranches final : public FSourceControlOperationBase
+{
+public:
+	// ISourceControlOperation interface
+	virtual FName GetName() const override;
+
+	virtual FText GetInProgressString() const override;
+
+	// Limit the list of branches to ones created from this date
+	FDateTime FromDate;
+
+	// List of branches found
+	TArray<FPlasticSourceControlBranchRef> Branches;
+};
+
+
+/**
+ * Internal operation used to switch the workspace to another branch
+*/
+class FPlasticSwitchToBranch final : public FSourceControlOperationBase
+{
+public:
+	// ISourceControlOperation interface
+	virtual FName GetName() const override;
+
+	virtual FText GetInProgressString() const override;
+
+	// Branch to switch the workspace to
+	FString BranchName;
+
+	/** List of files updated by the operation */
+	TArray<FString> UpdatedFiles;
+};
+
+
+/**
+ * Internal operation used to merge a branch into the current branch
+*/
+class FPlasticMergeBranch final : public FSourceControlOperationBase
+{
+public:
+	// ISourceControlOperation interface
+	virtual FName GetName() const override;
+
+	virtual FText GetInProgressString() const override;
+
+	// Branch to switch the workspace to
+	FString BranchName;
+
+	/** List of files updated by the operation */
+	TArray<FString> UpdatedFiles;
+};
+
+
+/**
+ * Internal operation used to create a branch
+*/
+class FPlasticCreateBranch final : public FSourceControlOperationBase
+{
+public:
+	// ISourceControlOperation interface
+	virtual FName GetName() const override;
+
+	virtual FText GetInProgressString() const override;
+
+	FString BranchName;
+	FString Comment;
+};
+
+
+/**
+ * Internal operation used to rename a branch
+*/
+class FPlasticRenameBranch final : public FSourceControlOperationBase
+{
+public:
+	// ISourceControlOperation interface
+	virtual FName GetName() const override;
+
+	virtual FText GetInProgressString() const override;
+
+	FString OldName;
+	FString NewName;
+};
+
+
+/**
+ * Internal operation used to deletes branches
+*/
+class FPlasticDeleteBranches final : public FSourceControlOperationBase
+{
+public:
+	// ISourceControlOperation interface
+	virtual FName GetName() const override;
+
+	virtual FText GetInProgressString() const override;
+
+	TArray<FString> BranchNames;
 };
 
 
@@ -238,7 +372,7 @@ public:
 	TArray<FPlasticSourceControlState> States;
 };
 
-/** Initialize a new Workspace and a new Repository */
+/** Create a new Workspace and a new Repository */
 class FPlasticMakeWorkspaceWorker final : public IPlasticSourceControlWorker
 {
 public:
@@ -252,6 +386,137 @@ public:
 	virtual bool UpdateStates() override;
 };
 
+/** Switch to Partial Workspace. */
+class FPlasticSwitchToPartialWorkspaceWorker final : public IPlasticSourceControlWorker
+{
+public:
+	explicit FPlasticSwitchToPartialWorkspaceWorker(FPlasticSourceControlProvider& InSourceControlProvider)
+		: IPlasticSourceControlWorker(InSourceControlProvider)
+	{}
+	virtual ~FPlasticSwitchToPartialWorkspaceWorker() = default;
+	// IPlasticSourceControlWorker interface
+	virtual FName GetName() const override;
+	virtual bool Execute(class FPlasticSourceControlCommand& InCommand) override;
+	virtual bool UpdateStates() override;
+
+public:
+	/** Temporary states for results */
+	TArray<FPlasticSourceControlState> States;
+};
+
+/** release or remove Lock(s) on file(s). */
+class FPlasticUnlockWorker final : public IPlasticSourceControlWorker
+{
+public:
+	explicit FPlasticUnlockWorker(FPlasticSourceControlProvider& InSourceControlProvider)
+		: IPlasticSourceControlWorker(InSourceControlProvider)
+	{}
+	virtual ~FPlasticUnlockWorker() = default;
+	// IPlasticSourceControlWorker interface
+	virtual FName GetName() const override;
+	virtual bool Execute(class FPlasticSourceControlCommand& InCommand) override;
+	virtual bool UpdateStates() override;
+
+public:
+	/** Temporary states for results */
+	TArray<FPlasticSourceControlState> States;
+};
+
+/** list branches. */
+class FPlasticGetBranchesWorker final : public IPlasticSourceControlWorker
+{
+public:
+	explicit FPlasticGetBranchesWorker(FPlasticSourceControlProvider& InSourceControlProvider)
+		: IPlasticSourceControlWorker(InSourceControlProvider)
+	{}
+	virtual ~FPlasticGetBranchesWorker() = default;
+	// IPlasticSourceControlWorker interface
+	virtual FName GetName() const override;
+	virtual bool Execute(class FPlasticSourceControlCommand& InCommand) override;
+	virtual bool UpdateStates() override;
+
+	// Current branch the workspace is on (at the end of the operation)
+	FString CurrentBranchName;
+};
+
+/** Switch workspace to another branch. */
+class FPlasticSwitchToBranchWorker final : public IPlasticSourceControlWorker
+{
+public:
+	explicit FPlasticSwitchToBranchWorker(FPlasticSourceControlProvider& InSourceControlProvider)
+		: IPlasticSourceControlWorker(InSourceControlProvider)
+	{}
+	virtual ~FPlasticSwitchToBranchWorker() = default;
+	// IPlasticSourceControlWorker interface
+	virtual FName GetName() const override;
+	virtual bool Execute(class FPlasticSourceControlCommand& InCommand) override;
+	virtual bool UpdateStates() override;
+
+public:
+	/** Temporary states for results */
+	TArray<FPlasticSourceControlState> States;
+};
+
+/** Merge a branch to the current branch. */
+class FPlasticMergeBranchWorker final : public IPlasticSourceControlWorker
+{
+public:
+	explicit FPlasticMergeBranchWorker(FPlasticSourceControlProvider& InSourceControlProvider)
+		: IPlasticSourceControlWorker(InSourceControlProvider)
+	{}
+	virtual ~FPlasticMergeBranchWorker() = default;
+	// IPlasticSourceControlWorker interface
+	virtual FName GetName() const override;
+	virtual bool Execute(class FPlasticSourceControlCommand& InCommand) override;
+	virtual bool UpdateStates() override;
+
+public:
+	/** Temporary states for results */
+	TArray<FPlasticSourceControlState> States;
+};
+
+/** Create a new child branch. */
+class FPlasticCreateBranchWorker final : public IPlasticSourceControlWorker
+{
+public:
+	explicit FPlasticCreateBranchWorker(FPlasticSourceControlProvider& InSourceControlProvider)
+		: IPlasticSourceControlWorker(InSourceControlProvider)
+	{}
+	virtual ~FPlasticCreateBranchWorker() = default;
+	// IPlasticSourceControlWorker interface
+	virtual FName GetName() const override;
+	virtual bool Execute(class FPlasticSourceControlCommand& InCommand) override;
+	virtual bool UpdateStates() override;
+};
+
+/** Rename a branch. */
+class FPlasticRenameBranchWorker final : public IPlasticSourceControlWorker
+{
+public:
+	explicit FPlasticRenameBranchWorker(FPlasticSourceControlProvider& InSourceControlProvider)
+		: IPlasticSourceControlWorker(InSourceControlProvider)
+	{}
+	virtual ~FPlasticRenameBranchWorker() = default;
+	// IPlasticSourceControlWorker interface
+	virtual FName GetName() const override;
+	virtual bool Execute(class FPlasticSourceControlCommand& InCommand) override;
+	virtual bool UpdateStates() override;
+};
+
+/** Delete branches. */
+class FPlasticDeleteBranchesWorker final : public IPlasticSourceControlWorker
+{
+public:
+	explicit FPlasticDeleteBranchesWorker(FPlasticSourceControlProvider& InSourceControlProvider)
+		: IPlasticSourceControlWorker(InSourceControlProvider)
+	{}
+	virtual ~FPlasticDeleteBranchesWorker() = default;
+	// IPlasticSourceControlWorker interface
+	virtual FName GetName() const override;
+	virtual bool Execute(class FPlasticSourceControlCommand& InCommand) override;
+	virtual bool UpdateStates() override;
+};
+
 /** Plastic update the workspace to latest changes */
 class FPlasticSyncWorker final : public IPlasticSourceControlWorker
 {
@@ -494,3 +759,30 @@ protected:
 	int32 ShelveId = ISourceControlState::INVALID_REVISION;
 };
 
+class FPlasticGetChangelistDetailsWorker final : public IPlasticSourceControlWorker
+{
+public:
+	FPlasticGetChangelistDetailsWorker(FPlasticSourceControlProvider& InSourceControlProvider)
+		: IPlasticSourceControlWorker(InSourceControlProvider)
+	{}
+	virtual ~FPlasticGetChangelistDetailsWorker() = default;
+
+	// IPlasticSourceControlWorker interface
+	virtual FName GetName() const override;
+	virtual bool Execute(FPlasticSourceControlCommand& InCommand) override;
+	virtual bool UpdateStates() override;
+};
+
+class FPlasticGetFileWorker final : public IPlasticSourceControlWorker
+{
+public:
+	FPlasticGetFileWorker(FPlasticSourceControlProvider& InSourceControlProvider)
+		: IPlasticSourceControlWorker(InSourceControlProvider)
+	{}
+	virtual ~FPlasticGetFileWorker() = default;
+
+	// IPlasticSourceControlWorker interface
+	virtual FName GetName() const override;
+	virtual bool Execute(FPlasticSourceControlCommand& InCommand) override;
+	virtual bool UpdateStates() override;
+};
diff --git a/Engine/Plugins/Developer/PlasticSourceControl/Source/PlasticSourceControl/Private/PlasticSourceControlProjectSettings.h b/Engine/Plugins/Developer/PlasticSourceControl/Source/PlasticSourceControl/Private/PlasticSourceControlProjectSettings.h
--- a/Engine/Plugins/Developer/PlasticSourceControl/Source/PlasticSourceControl/Private/PlasticSourceControlProjectSettings.h
+++ b/Engine/Plugins/Developer/PlasticSourceControl/Source/PlasticSourceControl/Private/PlasticSourceControlProjectSettings.h
@@ -6,26 +6,42 @@
 #include "Engine/DeveloperSettings.h"
 #include "PlasticSourceControlProjectSettings.generated.h"
 
-/** Project Settings for Plastic SCM Source Control. Saved in Config/DefaultEditor.ini */
-UCLASS(config = Editor, defaultconfig, meta = (DisplayName = "Source Control - Plastic SCM"))
+/** Project Settings for Unity Version Control (formerly Plastic SCM). Saved in Config/DefaultEditor.ini */
+UCLASS(config = Editor, defaultconfig, meta = (DisplayName = "Source Control - Unity Version Control"))
 class UPlasticSourceControlProjectSettings : public UDeveloperSettings
 {
 	GENERATED_BODY()
 
 public:
-	/** Map Plastic SCM user names (typically e-mail addresses or company domain names) to display names for brevity. */
-	UPROPERTY(config, EditAnywhere, Category = "Plastic SCM")
+	/** Map Unity Version Control user names (typically e-mail addresses or company domain names) to display names for brevity. */
+	UPROPERTY(config, EditAnywhere, Category = "Unity Version Control")
 	TMap<FString, FString> UserNameToDisplayName;
 
 	/** Hide the domain part of an username e-mail address (eg @gmail.com) if the UserNameToDisplayName map didn't match (enabled by default). */
-	UPROPERTY(config, EditAnywhere, Category = "Plastic SCM")
+	UPROPERTY(config, EditAnywhere, Category = "Unity Version Control")
 	bool bHideEmailDomainInUsername = true;
 
 	/** If enabled, you'll be prompted to check out changed files (enabled by default). Checkout is needed to work with Changelists. */
-	UPROPERTY(config, EditAnywhere, Category = "Plastic SCM")
+	UPROPERTY(config, EditAnywhere, Category = "Unity Version Control")
 	bool bPromptForCheckoutOnChange = true;
 
-	/** If a non-null value is set, limit the maximum number of revisions requested to Plastic SCM to display in the "History" window. */
-	UPROPERTY(config, EditAnywhere, Category = "Plastic SCM", meta=(ClampMin=0))
+	/** If a non-null value is set, limit the maximum number of revisions requested to Unity Version Control to display in the "History" window. */
+	UPROPERTY(config, EditAnywhere, Category = "Unity Version Control", meta = (ClampMin = 0))
 	int32 LimitNumberOfRevisionsInHistory = 50;
+
+	/** Show the repository when the branch is created (hidden by default) */
+	UPROPERTY(config, EditAnywhere, Category = "Unity Version Control")
+	bool bShowBranchRepositoryColumn = false;
+
+	/* Show the name of the creator of the branch */
+	UPROPERTY(config, EditAnywhere, Category = "Unity Version Control")
+	bool bShowBranchCreatedByColumn = true;
+
+	/* Show the date of creation of the branch */
+	UPROPERTY(config, EditAnywhere, Category = "Unity Version Control")
+	bool bShowBranchDateColumn = true;
+
+	/* Show the comment of the branch */
+	UPROPERTY(config, EditAnywhere, Category = "Unity Version Control")
+	bool bShowBranchCommentColumn = true;
 };
diff --git a/Engine/Plugins/Developer/PlasticSourceControl/Source/PlasticSourceControl/Private/PlasticSourceControlProvider.cpp b/Engine/Plugins/Developer/PlasticSourceControl/Source/PlasticSourceControl/Private/PlasticSourceControlProvider.cpp
--- a/Engine/Plugins/Developer/PlasticSourceControl/Source/PlasticSourceControl/Private/PlasticSourceControlProvider.cpp
+++ b/Engine/Plugins/Developer/PlasticSourceControl/Source/PlasticSourceControl/Private/PlasticSourceControlProvider.cpp
@@ -4,12 +4,14 @@
 
 #include "PlasticSourceControlChangelistState.h"
 #include "PlasticSourceControlCommand.h"
+#include "PlasticSourceControlModule.h"
 #include "PlasticSourceControlOperations.h"
 #include "PlasticSourceControlProjectSettings.h"
 #include "PlasticSourceControlSettings.h"
 #include "PlasticSourceControlShell.h"
 #include "PlasticSourceControlState.h"
 #include "PlasticSourceControlUtils.h"
+#include "PlasticSourceControlVersions.h"
 #include "SPlasticSourceControlSettings.h"
 
 #include "ISourceControlModule.h"
@@ -24,6 +26,8 @@
 #include "Misc/MessageDialog.h"
 #include "HAL/PlatformProcess.h"
 #include "Misc/QueuedThreadPool.h"
+#include "UObject/ObjectSaveContext.h"
+#include "UObject/SavePackage.h"
 
 #define LOCTEXT_NAMESPACE "PlasticSourceControl"
 
@@ -32,22 +36,32 @@ static FName ProviderName("Plastic SCM");
 FPlasticSourceControlProvider::FPlasticSourceControlProvider()
 {
 	PlasticSourceControlSettings.LoadSettings();
+
+	UPackage::PackageSavedWithContextEvent.AddRaw(this, &FPlasticSourceControlProvider::HandlePackageSaved);
+}
+
+FPlasticSourceControlProvider::~FPlasticSourceControlProvider()
+{
+	UPackage::PackageSavedWithContextEvent.RemoveAll(this);
 }
 
 void FPlasticSourceControlProvider::Init(bool bForceConnection)
 {
-	// Init() is called multiple times at startup: do not check Plastic SCM each time
+	// Init() is called multiple times at startup: do not check Unity Version Control each time
 	if (!bPlasticAvailable)
 	{
-		const TSharedPtr<IPlugin> Plugin = IPluginManager::Get().FindPlugin(TEXT("PlasticSourceControl"));
+		const TSharedPtr<IPlugin> Plugin = FPlasticSourceControlModule::GetPlugin();
 		if (Plugin.IsValid())
 		{
 			PluginVersion = Plugin->GetDescriptor().VersionName;
-			UE_LOG(LogSourceControl, Log, TEXT("Plastic SCM plugin '%s'"), *PluginVersion);
+			UE_LOG(LogSourceControl, Log, TEXT("Unity Version Control (formerly Plastic SCM) plugin %s"), *PluginVersion);
 		}
 
 		CheckPlasticAvailability();
 
+		FMessageLog("SourceControl").Info(FText::Format(LOCTEXT("PluginVersion", "Unity Version Control (formerly Plastic SCM) {0} (plugin {1})"),
+			FText::FromString(PlasticScmVersion.String), FText::FromString(PluginVersion)));
+
 		// Override the source control logs verbosity level if needed based on settings
 		if (AccessSettings().GetEnableVerboseLogs())
 		{
@@ -57,14 +71,10 @@ void FPlasticSourceControlProvider::Init(bool bForceConnection)
 
 	if (bForceConnection && bPlasticAvailable && bWorkspaceFound && !bServerAvailable)
 	{
-		// Execute a 'checkconnection' command to set bServerAvailable based on the connectivity of the server
 		TArray<FString> InfoMessages, ErrorMessages;
 		TArray<FString> Parameters;
-		if (PlasticSourceControlUtils::GetWorkspaceInfo(BranchName, RepositoryName, ServerUrl, ErrorMessages))
-		{
-			Parameters.Add(FString::Printf(TEXT("--server=%s"), *ServerUrl));
-		}
-		bServerAvailable = PlasticSourceControlUtils::RunCommand(TEXT("checkconnection"), Parameters, TArray<FString>(), InfoMessages, ErrorMessages);
+		// Execute a 'checkconnection' command to set bServerAvailable based on the connectivity of the server
+		bServerAvailable = PlasticSourceControlUtils::RunCheckConnection(BranchName, RepositoryName, ServerUrl, InfoMessages, ErrorMessages);
 		if (!bServerAvailable)
 		{
 			FMessageLog SourceControlLog("SourceControl");
@@ -83,7 +93,7 @@ void FPlasticSourceControlProvider::CheckPlasticAvailability()
 	{
 		bPlasticAvailable = false;
 
-		// Try to find Plastic binary, and update settings accordingly
+		// Try to find cm binary, and update settings accordingly
 		PathToPlasticBinary = PlasticSourceControlUtils::FindPlasticBinaryPath();
 		if (!PathToPlasticBinary.IsEmpty())
 		{
@@ -93,10 +103,9 @@ void FPlasticSourceControlProvider::CheckPlasticAvailability()
 
 	if (!PathToPlasticBinary.IsEmpty())
 	{
-		// Find the path to the root Plastic directory (if any, else uses the ProjectDir)
 		const FString PathToProjectDir = FPaths::ConvertRelativePathToFull(FPaths::ProjectDir());
 
-		// Launch the Plastic SCM cli shell on the background to issue all commands during this session
+		// Launch the Unity Version Control cli shell on the background to issue all commands during this session
 		bPlasticAvailable = PlasticSourceControlShell::Launch(PathToPlasticBinary, PathToProjectDir);
 		if (!bPlasticAvailable)
 		{
@@ -112,12 +121,12 @@ void FPlasticSourceControlProvider::CheckPlasticAvailability()
 		FString ActualPathToPlasticBinary;
 		PlasticSourceControlUtils::GetCmLocation(ActualPathToPlasticBinary);
 
+		// Find the path to the root Plastic directory (if any, else uses the ProjectDir)
 		bWorkspaceFound = PlasticSourceControlUtils::GetWorkspacePath(PathToProjectDir, PathToWorkspaceRoot);
 
-
 		bUsesLocalReadOnlyState = PlasticSourceControlUtils::GetConfigSetFilesAsReadOnly();
 
-		// Get user name (from the global Plastic SCM client config)
+		// Get user name (from the global Unity Version Control client config)
 		PlasticSourceControlUtils::GetUserName(UserName);
 
 		// Register Console Commands
@@ -129,6 +138,8 @@ void FPlasticSourceControlProvider::CheckPlasticAvailability()
 			FFormatNamedArguments Args;
 			Args.Add(TEXT("WorkspacePath"), FText::FromString(PathToWorkspaceRoot));
 			FMessageLog("SourceControl").Info(FText::Format(LOCTEXT("NotInAWorkspace", "{WorkspacePath} is not in a workspace."), Args));
+
+			ServerUrl = PlasticSourceControlUtils::GetConfigDefaultRepServer();
 		}
 	}
 }
@@ -184,6 +195,26 @@ TSharedRef<FPlasticSourceControlChangelistState, ESPMode::ThreadSafe> FPlasticSo
 	}
 }
 
+// Note: called once for each asset being saved, which can be hundreds in the case of a map using One File Per Actor (OFPA) in UE5
+void FPlasticSourceControlProvider::HandlePackageSaved(const FString& InPackageFilename, UPackage* InPackage, FObjectPostSaveContext InObjectSaveContext)
+{
+	const FString AbsoluteFilename = FPaths::ConvertRelativePathToFull(InPackageFilename);
+	auto FileState = GetStateInternal(AbsoluteFilename);
+
+	// Note: the Editor doesn't ask to refresh the source control status of an asset after it is saved, only *before* (to check that it's possible to save)
+	// So when an asset with no change is saved, update its state in cache to record the fact that the asset is now changed.
+	if (FileState->WorkspaceState == EWorkspaceState::Controlled)
+	{
+		// Note that updating the state in cache isn't enough to refresh the status icon in the Content Browser (since the Editor isn't made aware of the change)
+		// but source control operations are working as expected (eg. "Checkin" and "Revert" are available in the context menu)
+		FileState->WorkspaceState = EWorkspaceState::Changed; // The icon will only appears later when the UI is refreshed (eg switching directory in the Content Browser)
+	}
+	else if (FileState->WorkspaceState == EWorkspaceState::CheckedOutUnchanged)
+	{
+		FileState->WorkspaceState = EWorkspaceState::CheckedOutChanged; // In this case the "CheckedOut" icon is already displayed (both states are using the same status icon)
+	}
+}
+
 FText FPlasticSourceControlProvider::GetStatusText() const
 {
 	FFormatNamedArguments Args;
@@ -192,6 +223,8 @@ FText FPlasticSourceControlProvider::GetStatusText() const
 	Args.Add(TEXT("WorkspacePath"), FText::FromString(PathToWorkspaceRoot));
 	Args.Add(TEXT("WorkspaceName"), FText::FromString(WorkspaceName));
 	Args.Add(TEXT("BranchName"), FText::FromString(BranchName));
+	Args.Add(TEXT("RepositoryName"), FText::FromString(RepositoryName));
+	Args.Add(TEXT("ServerUrl"), FText::FromString(ServerUrl));
 	// Detect special case for a partial checkout (CS:-1 in Gluon mode)!
 	if (IsPartialWorkspace())
 	{
@@ -223,7 +256,7 @@ FText FPlasticSourceControlProvider::GetStatusText() const
 	}
 	Args.Add(TEXT("ErrorText"), FormattedError);
 
-	return FText::Format(LOCTEXT("PlasticStatusText", "{ErrorText}Plastic SCM {PlasticScmVersion}  (plugin v{PluginVersion})\nWorkspace: {WorkspaceName}  ({WorkspacePath})\n{BranchName}\nChangeset: {ChangesetNumber}\nUser: '{UserName}'  {DisplayName}"), Args);
+	return FText::Format(LOCTEXT("PlasticStatusText", "{ErrorText}Unity Version Control (formerly Plastic SCM) {PlasticScmVersion}\t(plugin v{PluginVersion})\nWorkspace: {WorkspaceName}  ({WorkspacePath})\nBranch: {BranchName}@{RepositoryName}@{ServerUrl}\nChangeset: {ChangesetNumber}\nUser: '{UserName}'  {DisplayName}"), Args);
 }
 
 TMap<ISourceControlProvider::EStatus, FString> FPlasticSourceControlProvider::GetStatus() const
@@ -232,13 +265,13 @@ TMap<ISourceControlProvider::EStatus, FString> FPlasticSourceControlProvider::Ge
 	Result.Add(EStatus::Enabled, IsEnabled() ? TEXT("Yes") : TEXT("No") );
 	Result.Add(EStatus::Connected, (IsEnabled() && IsAvailable()) ? TEXT("Yes") : TEXT("No") );
 	Result.Add(EStatus::User, UserName);
-	
+
 	Result.Add(EStatus::ScmVersion, PlasticScmVersion.String);
 	Result.Add(EStatus::PluginVersion, PluginVersion);
 	Result.Add(EStatus::WorkspacePath, PathToWorkspaceRoot);
 	Result.Add(EStatus::Workspace, WorkspaceName);
 	Result.Add(EStatus::Branch, BranchName);
-	if (-1 != ChangesetNumber)
+	if (IsPartialWorkspace())
 	{
 		Result.Add(EStatus::Changeset, FString::Printf(TEXT("%d"), ChangesetNumber));
 	}
@@ -262,6 +295,17 @@ const FName& FPlasticSourceControlProvider::GetName(void) const
 	return ProviderName;
 }
 
+FString FPlasticSourceControlProvider::GetCloudOrganization() const
+{
+	const int32 CloudIndex = ServerUrl.Find(TEXT("@cloud"));
+	if (CloudIndex > INDEX_NONE)
+	{
+		return ServerUrl.Left(CloudIndex);
+	}
+
+	return FString();
+}
+
 void FPlasticSourceControlProvider::SetLastErrors(const TArray<FString>& InErrors)
 {
 	FScopeLock Lock(&LastErrorsCriticalSection);
@@ -387,7 +431,12 @@ ECommandResult::Type FPlasticSourceControlProvider::Execute(const FSourceControl
 		FFormatNamedArguments Arguments;
 		Arguments.Add(TEXT("OperationName"), FText::FromName(InOperation->GetName()));
 		Arguments.Add(TEXT("ProviderName"), FText::FromName(GetName()));
-		FMessageLog("SourceControl").Error(FText::Format(LOCTEXT("UnsupportedOperation", "Operation '{OperationName}' not supported by source control provider '{ProviderName}'"), Arguments));
+		FText Message = FText::Format(LOCTEXT("UnsupportedOperation", "Operation '{OperationName}' not supported by revision control provider '{ProviderName}'"), Arguments);
+
+		FMessageLog("SourceControl").Error(Message);
+		InOperation->AddErrorMessge(Message);
+
+		InOperationCompleteDelegate.ExecuteIfBound(InOperation, ECommandResult::Failed);
 		return ECommandResult::Failed;
 	}
 
@@ -436,6 +485,7 @@ bool FPlasticSourceControlProvider::UsesLocalReadOnlyState() const
 
 bool FPlasticSourceControlProvider::UsesChangelists() const
 {
+	// We don't want to show ChangeList column anymore (Unity Version Control term would be ChangeSet) BUT we need this to display the changelists in the source control menu
 	return true;
 }
 
@@ -451,9 +501,10 @@ bool FPlasticSourceControlProvider::UsesCheckout() const
 
 bool FPlasticSourceControlProvider::UsesFileRevisions() const
 {
-	// Only a Partial/Gluon workspace can sync/update files individually, operating on revisions (can use the context menu)
-	// while a regular workspace can only update all the files as a whole, operating at the changeset level (requires the global menu)
-	return IsPartialWorkspace();
+	// This API introduced in UE5.1 is still broken as of UE5.3
+	// (preventing the user to use the source control context menu for checkin if returning false)
+	// return IsPartialWorkspace();
+	return true;
 }
 
 bool FPlasticSourceControlProvider::UsesSnapshots() const
@@ -494,16 +545,23 @@ void FPlasticSourceControlProvider::RegisterWorker(const FName& InName, const FG
 
 void FPlasticSourceControlProvider::OutputCommandMessages(const FPlasticSourceControlCommand& InCommand) const
 {
+	// Note: the Perforce provider added a way to call OutputCommandMessages() from FPerforceSourceControlProvider::TryToDownloadFileFromBackgroundThread()
+	// see Commit 18dc0643 by paul chipchase, 05/17/2021 11:06 AM [CL 16346666 by paul chipchase in ue5-main branch]
+	// but we don't have this new API yet so we don't need this thread safety
+	check(IsInGameThread()); // On the game thread we can use FMessageLog
+
 	FMessageLog SourceControlLog("SourceControl");
 
 	for (const FString& ErrorMessage : InCommand.ErrorMessages)
 	{
-		SourceControlLog.Error(FText::FromString(ErrorMessage));
+		SourceControlLog.Error(FText::Format(LOCTEXT("OutputCommandMessagesFormatError", "Command: {0}, Error: {1}"),
+			FText::FromName(InCommand.Operation->GetName()), FText::FromString(ErrorMessage)));
 	}
 
 	for (const FString& InfoMessage : InCommand.InfoMessages)
 	{
-		SourceControlLog.Info(FText::FromString(InfoMessage));
+		SourceControlLog.Info(FText::Format(LOCTEXT("OutputCommandMessagesFormatInfo", "Command: {0}, Info: {1}"),
+			FText::FromName(InCommand.Operation->GetName()), FText::FromString(InfoMessage)));
 	}
 }
 
@@ -524,19 +582,25 @@ void FPlasticSourceControlProvider::UpdateWorkspaceStatus(const class FPlasticSo
 		{
 			if (bPlasticAvailable)
 			{
-				if (PlasticScmVersion < PlasticSourceControlUtils::GetOldestSupportedPlasticScmVersion())
+				if (PlasticScmVersion < PlasticSourceControlVersions::OldestSupported)
 				{
 					FFormatNamedArguments Args;
 					Args.Add(TEXT("PlasticScmVersion"), FText::FromString(PlasticScmVersion.String));
-					Args.Add(TEXT("OldestSupportedPlasticScmVersion"), FText::FromString(PlasticSourceControlUtils::GetOldestSupportedPlasticScmVersion().String));
-					const FText UnsuportedVersionWarning = FText::Format(LOCTEXT("Plastic_UnsuportedVersion", "Plastic SCM {PlasticScmVersion} is not supported anymore by this plugin.\nPlastic SCM {OldestSupportedPlasticScmVersion} or a more recent version is required.\nPlease upgrade to the latest version."), Args);
-					FMessageLog("SourceControl").Warning(UnsuportedVersionWarning);
-					FMessageDialog::Open(EAppMsgType::Ok, UnsuportedVersionWarning);
+					Args.Add(TEXT("OldestSupportedPlasticScmVersion"), FText::FromString(PlasticSourceControlVersions::OldestSupported.String));
+					const FText UnsupportedVersionWarning = FText::Format(LOCTEXT("Plastic_UnsupportedVersion", "Unity Version Control {PlasticScmVersion} is not supported anymore by this plugin.\nUnity Version Control {OldestSupportedPlasticScmVersion} or a more recent version is required.\nPlease upgrade to the latest version."), Args);
+					FMessageLog("SourceControl").Warning(UnsupportedVersionWarning);
+					FMessageDialog::Open(
+						EAppMsgCategory::Warning,
+						EAppMsgType::Ok, UnsupportedVersionWarning
+						, LOCTEXT("Plastic_UnsuportedVersionTitle", "Unsupported version!")
+					);
 				}
 			}
 			else if (InCommand.ErrorMessages.Num() > 0)
 			{
-				FMessageDialog::Open(EAppMsgType::Ok, FText::FromString(InCommand.ErrorMessages[0]));
+				FMessageDialog::Open(
+					EAppMsgCategory::Error,
+					EAppMsgType::Ok, FText::FromString(InCommand.ErrorMessages[0]));
 			}
 		}
 
@@ -600,7 +664,7 @@ void FPlasticSourceControlProvider::Tick()
 
 			if (Command.Files.Num() > 1)
 			{
-				UE_LOG(LogSourceControl, Log, TEXT("%s of %d files processed in %.3lfs"), *Command.Operation->GetName().ToString(), Command.Files.Num(), (FPlatformTime::Seconds() - Command.StartTimestamp));
+				UE_LOG(LogSourceControl, Log, TEXT("%s of %d items processed in %.3lfs"), *Command.Operation->GetName().ToString(), Command.Files.Num(), (FPlatformTime::Seconds() - Command.StartTimestamp));
 			}
 			else if (Command.Files.Num() == 1)
 			{
@@ -612,8 +676,7 @@ void FPlasticSourceControlProvider::Tick()
 			}
 
 			// run the completion delegate callback if we have one bound
-			ECommandResult::Type Result = Command.bCommandSuccessful ? ECommandResult::Succeeded : ECommandResult::Failed;
-			Command.OperationCompleteDelegate.ExecuteIfBound(Command.Operation, Result);
+			Command.ReturnResults();
 
 			// commands that are left in the array during a tick need to be deleted
 			if (Command.bAutoDelete)
@@ -638,6 +701,10 @@ void FPlasticSourceControlProvider::Tick()
 TArray<TSharedRef<ISourceControlLabel>> FPlasticSourceControlProvider::GetLabels(const FString& InMatchingSpec) const
 {
 	TArray< TSharedRef<ISourceControlLabel> > Tags;
+
+	// NOTE list labels. Called by CrashDebugHelper() (to remote debug Engine crash)
+	//					 and by SourceControlHelpers::AnnotateFile() (to add source file to report)
+	// Reserved for internal use by Epic Games with Perforce only
 	return Tags;
 }
 
@@ -702,9 +769,6 @@ ECommandResult::Type FPlasticSourceControlProvider::ExecuteSynchronousCommand(FP
 		}
 		else
 		{
-			// TODO If the command failed, inform the user that they need to try again (see Perforce)
-			// FMessageDialog::Open( EAppMsgType::Ok, LOCTEXT("Plastic_ServerUnresponsive", "Plastic server is unresponsive. Please check your connection and try again.") );
-
 			UE_LOG(LogSourceControl, Error, TEXT("Command '%s' Failed!"), *InCommand.Operation->GetName().ToString());
 		}
 	}
@@ -733,8 +797,17 @@ ECommandResult::Type FPlasticSourceControlProvider::IssueCommand(FPlasticSourceC
 	}
 	else
 	{
-		return ECommandResult::Failed;
+		// NOTE: the Perforce & Subversion providers implement this, but looking at Git history of the code I think it has never been used in UE4
+		// If we need to support this, we will need to know the use cases in order to test it
+		FText Message(LOCTEXT("NoSCCThreads", "There are no threads available to process the revision control command."));
+
+		FMessageLog("SourceControl").Error(Message);
+		InCommand.bCommandSuccessful = false;
+		InCommand.Operation->AddErrorMessge(Message);
+
+		return InCommand.ReturnResults();
 	}
 }
 
 #undef LOCTEXT_NAMESPACE
+
