PR #11295: [Source Control] Update Plastic SCM provider to 1.9.0 (from 1.8.0)

diff --git a/Engine/Plugins/Developer/PlasticSourceControl/Source/PlasticSourceControl/Private/PlasticSourceControlProvider.h b/Engine/Plugins/Developer/PlasticSourceControl/Source/PlasticSourceControl/Private/PlasticSourceControlProvider.h
--- a/Engine/Plugins/Developer/PlasticSourceControl/Source/PlasticSourceControl/Private/PlasticSourceControlProvider.h
+++ b/Engine/Plugins/Developer/PlasticSourceControl/Source/PlasticSourceControl/Private/PlasticSourceControlProvider.h
@@ -18,8 +18,9 @@ DECLARE_DELEGATE_RetVal_OneParam(FPlasticSourceControlWorkerRef, FGetPlasticSour
 class FPlasticSourceControlProvider : public ISourceControlProvider
 {
 public:
-	/** Constructor */
+	/** Constructor & destructor */
 	FPlasticSourceControlProvider();
+	~FPlasticSourceControlProvider();
 
 	/* ISourceControlProvider implementation */
 	virtual void Init(bool bForceConnection = true) override;
@@ -38,7 +39,7 @@ public:
 	virtual FDelegateHandle RegisterSourceControlStateChanged_Handle(const FSourceControlStateChanged::FDelegate& SourceControlStateChanged) override;
 	virtual void UnregisterSourceControlStateChanged_Handle(FDelegateHandle Handle) override;
 	virtual ECommandResult::Type Execute(const FSourceControlOperationRef& InOperation, FSourceControlChangelistPtr InChangelist, const TArray<FString>& InFiles, EConcurrency::Type InConcurrency = EConcurrency::Synchronous, const FSourceControlOperationComplete& InOperationCompleteDelegate = FSourceControlOperationComplete() ) override;
-	virtual bool CanExecuteOperation( const FSourceControlOperationRef& InOperation ) const override;
+	virtual bool CanExecuteOperation(const FSourceControlOperationRef& InOperation) const override;
 	virtual bool CanCancelOperation(const FSourceControlOperationRef& InOperation) const override;
 	virtual void CancelOperation(const FSourceControlOperationRef& InOperation) override;
 	virtual bool UsesLocalReadOnlyState() const override;
@@ -100,7 +101,7 @@ public:
 		return RepositoryName;
 	}
 
-	/** Get the Plastic current server URL */
+	/** Get the Plastic current server URL. See also GetCloudOrganization() */
 	inline const FString& GetServerUrl() const
 	{
 		return ServerUrl;
@@ -111,6 +112,10 @@ public:
 	{
 		return BranchName;
 	}
+	inline void SetBranchName(const FString& InBranchName)
+	{
+		BranchName = InBranchName;
+	}
 
 	/** Get the current Changeset Number */
 	inline int32 GetChangesetNumber() const
@@ -124,18 +129,21 @@ public:
 		return (ChangesetNumber == -1);
 	}
 
-	/** Version of the Plastic SCM executable used */
+	/** Version of the Unity Version Control executable used */
 	inline const FSoftwareVersion& GetPlasticScmVersion() const
 	{
 		return PlasticScmVersion;
 	}
 
-	/** Version of the Plastic SCM plugin */
+	/** Version of the Unity Version Control plugin */
 	const FString& GetPluginVersion() const
 	{
 		return PluginVersion;
 	}
 
+	/** Return the name of the cloud organization from the ServerUrl if applicable (MyOrganization@cloud) or an empty string */
+	FString GetCloudOrganization() const;
+
 	/** Set list of error messages that occurred after last Plastic command */
 	void SetLastErrors(const TArray<FString>& InErrors);
 
@@ -189,7 +197,7 @@ private:
 	/** Indicates if source control integration is available or not. */
 	bool bServerAvailable = false;
 
-	/** Whether Plastic SCM is configured to uses local read-only state to signal whether a file is editable ("SetFilesAsReadOnly" in client.conf) */
+	/** Whether Unity Version Control is configured to uses local read-only state to signal whether a file is editable ("SetFilesAsReadOnly" in client.conf) */
 	bool bUsesLocalReadOnlyState = false;
 
 	/** Critical section for thread safety of error messages that occurred after last Plastic command */
@@ -212,10 +220,13 @@ private:
 	/** Update workspace status on Connect and UpdateStatus operations */
 	void UpdateWorkspaceStatus(const class FPlasticSourceControlCommand& InCommand);
 
-	/** Version of the Plastic SCM executable used */
+	/** Called after a package has been saved to disk, to update the source control cache */
+	void HandlePackageSaved(const FString& InPackageFilename, UPackage* InPackage, FObjectPostSaveContext InObjectSaveContext);
+
+	/** Version of the Unity Version Control executable used */
 	FSoftwareVersion PlasticScmVersion;
 
-	/** Version of the Plastic SCM plugin */
+	/** Version of the Unity Version Control plugin */
 	FString PluginVersion;
 
 	/** Path to the root of the Plastic workspace: can be the GameDir itself, or any parent directory (found by the "Connect" operation) */
diff --git a/Engine/Plugins/Developer/PlasticSourceControl/Source/PlasticSourceControl/Private/PlasticSourceControlRevision.cpp b/Engine/Plugins/Developer/PlasticSourceControl/Source/PlasticSourceControl/Private/PlasticSourceControlRevision.cpp
--- a/Engine/Plugins/Developer/PlasticSourceControl/Source/PlasticSourceControl/Private/PlasticSourceControlRevision.cpp
+++ b/Engine/Plugins/Developer/PlasticSourceControl/Source/PlasticSourceControl/Private/PlasticSourceControlRevision.cpp
@@ -2,10 +2,8 @@
 
 #include "PlasticSourceControlRevision.h"
 #include "PlasticSourceControlModule.h"
-#include "PlasticSourceControlProvider.h"
 #include "PlasticSourceControlState.h"
 #include "PlasticSourceControlUtils.h"
-#include "SPlasticSourceControlSettings.h"
 
 #include "HAL/FileManager.h"
 #include "Misc/Paths.h"
@@ -15,15 +13,10 @@
 
 bool FPlasticSourceControlRevision::Get(FString& InOutFilename, EConcurrency::Type InConcurrency /* = EConcurrency::Synchronous */) const
 {
-	if (InConcurrency != EConcurrency::Synchronous)
-	{
-		UE_LOG(LogSourceControl, Warning, TEXT("Only EConcurrency::Synchronous is tested/supported for this operation."));
-	}
-
 	// if a filename for the temp file wasn't supplied generate a unique-ish one
-	if (InOutFilename.Len() == 0)
+	if (InOutFilename.IsEmpty())
 	{
-		// create the diff dir if we don't already have it (Plastic wont)
+		// create the diff dir if we don't already have it
 		IFileManager::Get().MakeDirectory(*FPaths::DiffDir(), true);
 		// create a unique temp file name based on the unique revision Id
 		FString TempFileName;
@@ -31,6 +24,10 @@ bool FPlasticSourceControlRevision::Get(FString& InOutFilename, EConcurrency::Ty
 		{
 			TempFileName = FString::Printf(TEXT("%stemp-sh%d-%s"), *FPaths::DiffDir(), ShelveId, *FPaths::GetCleanFilename(Filename));
 		}
+		else if (RevisionId != ISourceControlState::INVALID_REVISION)
+		{
+			TempFileName = FString::Printf(TEXT("%stemp-rev%d-%s"), *FPaths::DiffDir(), RevisionId, *FPaths::GetCleanFilename(Filename));
+		}
 		else
 		{
 			TempFileName = FString::Printf(TEXT("%stemp-cs%d-%s"), *FPaths::DiffDir(), ChangesetNumber, *FPaths::GetCleanFilename(Filename));
@@ -38,15 +35,13 @@ bool FPlasticSourceControlRevision::Get(FString& InOutFilename, EConcurrency::Ty
 		InOutFilename = FPaths::ConvertRelativePathToFull(TempFileName);
 	}
 
-	bool bCommandSuccessful;
+	bool bCommandSuccessful = false;
 	if (FPaths::FileExists(InOutFilename))
 	{
 		bCommandSuccessful = true; // if the temp file already exists, reuse it directly
 	}
-	else if (State)
+	else
 	{
-		const FString& PathToPlasticBinary = FPlasticSourceControlModule::Get().GetProvider().AccessSettings().GetBinaryPath();
-
 		FString RevisionSpecification;
 		if (ShelveId != ISourceControlState::INVALID_REVISION)
 		{
@@ -54,28 +49,48 @@ bool FPlasticSourceControlRevision::Get(FString& InOutFilename, EConcurrency::Ty
 			// Note: the plugin doesn't support shelves on Xlinks (no known RepSpec)
 			RevisionSpecification = FString::Printf(TEXT("rev:%s#sh:%d"), *Filename, ShelveId);
 		}
-		else
+		else if (RevisionId != ISourceControlState::INVALID_REVISION)
+		{
+			// Format the revision specification of the file, like rev:revid:920
+			RevisionSpecification = FString::Printf(TEXT("rev:revid:%d"), RevisionId);
+		}
+		else if (State)
 		{
 			// Format the revision specification of the checked-in file, like rev:Content/BP.uasset#cs:12@repo@server:8087
 			RevisionSpecification = FString::Printf(TEXT("rev:%s#cs:%d@%s"), *Filename, ChangesetNumber, *State->RepSpec);
 		}
-		bCommandSuccessful = PlasticSourceControlUtils::RunDumpToFile(PathToPlasticBinary, RevisionSpecification, InOutFilename);
-	}
-	else
-	{
-		UE_LOG(LogSourceControl, Error, TEXT("Revision(%s %d): unknown state!"), *Filename, RevisionId);
-		bCommandSuccessful = false;
+		else
+		{
+			UE_LOG(LogSourceControl, Error, TEXT("Unknown revision for %s!"), *Filename);
+		}
+
+		if (!RevisionSpecification.IsEmpty())
+		{
+			bCommandSuccessful = PlasticSourceControlUtils::RunGetFile(RevisionSpecification, InOutFilename);
+		}
+		if (!bCommandSuccessful && FPaths::FileExists(InOutFilename))
+		{
+			// On error, delete the temp file if it was created
+			IFileManager::Get().Delete(*InOutFilename);
+		}
 	}
 	return bCommandSuccessful;
 }
 
-bool FPlasticSourceControlRevision::GetAnnotated( TArray<FAnnotationLine>& OutLines ) const
+bool FPlasticSourceControlRevision::GetAnnotated(TArray<FAnnotationLine>& OutLines) const
 {
+	// NOTE GetAnnotated: called only by SourceControlHelpers::AnnotateFile(),
+	//      called only by ICrashDebugHelper::AddAnnotatedSourceToReport() using a changelist/check identifier
+	//      called only by FCrashDebugHelperWindows::CreateMinidumpDiagnosticReport() (and Mac) to Extract annotated lines from a source file stored in Perforce, and add to the crash report.
+	//      called by - MinidumpDiagnosticsApp RunMinidumpDiagnostics() for Perforce ONLY "MinidumpDiagnostics.exe <Crash.dmp> [-Annotate] [-SyncSymbols] [-SyncMicrosoftSymbols]"
+	//                - FWindowsErrorReport::DiagnoseReport() (and Mac)
+	// Reserved for internal use by Epic Games with Perforce only
 	return false;
 }
 
-bool FPlasticSourceControlRevision::GetAnnotated( FString& InOutFilename ) const
+bool FPlasticSourceControlRevision::GetAnnotated(FString& InOutFilename) const
 {
+	// NOTE: Unused, only the above method is called by the Editor
 	return false;
 }
 
diff --git a/Engine/Plugins/Developer/PlasticSourceControl/Source/PlasticSourceControl/Private/PlasticSourceControlRevision.h b/Engine/Plugins/Developer/PlasticSourceControl/Source/PlasticSourceControl/Private/PlasticSourceControlRevision.h
--- a/Engine/Plugins/Developer/PlasticSourceControl/Source/PlasticSourceControl/Private/PlasticSourceControlRevision.h
+++ b/Engine/Plugins/Developer/PlasticSourceControl/Source/PlasticSourceControl/Private/PlasticSourceControlRevision.h
@@ -4,6 +4,7 @@
 
 #include "CoreMinimal.h"
 #include "ISourceControlRevision.h"
+#include "ISourceControlState.h"
 
 class FPlasticSourceControlState;
 
@@ -12,8 +13,8 @@ class FPlasticSourceControlRevision : public ISourceControlRevision
 {
 public:
 	FPlasticSourceControlRevision()
-		: ChangesetNumber(0)
-		, RevisionId(0)
+		: ChangesetNumber(ISourceControlState::INVALID_REVISION)
+		, RevisionId(ISourceControlState::INVALID_REVISION)
 		, Date(0)
 		, FileSize(0)
 	{
diff --git a/Engine/Plugins/Developer/PlasticSourceControl/Source/PlasticSourceControl/Private/PlasticSourceControlShell.cpp b/Engine/Plugins/Developer/PlasticSourceControl/Source/PlasticSourceControl/Private/PlasticSourceControlShell.cpp
--- a/Engine/Plugins/Developer/PlasticSourceControl/Source/PlasticSourceControl/Private/PlasticSourceControlShell.cpp
+++ b/Engine/Plugins/Developer/PlasticSourceControl/Source/PlasticSourceControl/Private/PlasticSourceControlShell.cpp
@@ -2,6 +2,7 @@
 
 #include "PlasticSourceControlShell.h"
 
+#include "Notification.h"
 #include "PlasticSourceControlModule.h"
 #include "PlasticSourceControlProvider.h"
 
@@ -11,9 +12,6 @@
 #include "HAL/PlatformProcess.h"
 #include "HAL/PlatformTime.h"
 
-#include "Framework/Notifications/NotificationManager.h"
-#include "Widgets/Notifications/SNotificationList.h"
-
 #if PLATFORM_LINUX
 #include <sys/ioctl.h>
 #endif
@@ -42,6 +40,9 @@ static FCriticalSection	ShellCriticalSection;
 static size_t			ShellCommandCounter = -1;
 static double			ShellCumulatedTime = 0.;
 
+// Whether we already ran a status command to warm up the current shell process
+static bool             bShellIsWarmedUp = false;
+
 // Internal function to cleanup (called under the critical section)
 static void _CleanupBackgroundCommandLineShell()
 {
@@ -51,7 +52,17 @@ static void _CleanupBackgroundCommandLineShell()
 	ShellInputPipeRead = ShellInputPipeWrite = nullptr;
 }
 
-// Internal function to launch the Plastic SCM background 'cm' process in interactive shell mode (called under the critical section)
+static bool _GetShellIsWarmedUp()
+{
+	return bShellIsWarmedUp;
+}
+
+static void _SetShellIsWarmedUp()
+{
+	bShellIsWarmedUp = true;
+}
+
+// Internal function to launch the Unity Version Control background 'cm' process in interactive shell mode (called under the critical section)
 static bool _StartBackgroundPlasticShell(const FString& InPathToPlasticBinary, const FString& InWorkingDirectory)
 {
 	TRACE_CPUPROFILER_EVENT_SCOPE(PlasticSourceControlShell::_StartBackgroundPlasticShell);
@@ -62,6 +73,8 @@ static bool _StartBackgroundPlasticShell(const FString& InPathToPlasticBinary, c
 	const bool bLaunchHidden = true;				// the new process will be minimized in the task bar
 	const bool bLaunchReallyHidden = bLaunchHidden; // the new process will not have a window or be in the task bar
 
+	bShellIsWarmedUp = false;
+
 	const double StartTimestamp = FPlatformTime::Seconds();
 
 	verify(FPlatformProcess::CreatePipe(ShellOutputPipeRead, ShellOutputPipeWrite, false));	// For reading outputs from cm shell child process
@@ -69,7 +82,7 @@ static bool _StartBackgroundPlasticShell(const FString& InPathToPlasticBinary, c
 
 #if PLATFORM_WINDOWS
 	ShellProcessHandle = FPlatformProcess::CreateProc(*InPathToPlasticBinary, *FullCommand, bLaunchDetached, bLaunchHidden, bLaunchReallyHidden, nullptr, 0, *InWorkingDirectory, ShellOutputPipeWrite, ShellInputPipeRead);
-#else // PLATFORM_LINUX
+#else // PLATFORM_MAC or PLATFORM_LINUX
 	// Update working directory
 	char OriginalWorkingDirectory[PATH_MAX];
 	getcwd(OriginalWorkingDirectory, PATH_MAX);
@@ -83,7 +96,7 @@ static bool _StartBackgroundPlasticShell(const FString& InPathToPlasticBinary, c
 
 	if (!ShellProcessHandle.IsValid())
 	{
-		UE_LOG(LogSourceControl, Warning, TEXT("Failed to launch 'cm shell'")); // not a bug, just no Plastic SCM cli found
+		UE_LOG(LogSourceControl, Warning, TEXT("Failed to launch 'cm shell'")); // not a bug, just no Unity Version Control cli found
 		_CleanupBackgroundCommandLineShell();
 	}
 	else
@@ -147,18 +160,6 @@ static void _RestartBackgroundCommandLineShell(const bool bInForceExit = false)
 	_StartBackgroundPlasticShell(PathToPlasticBinary, WorkingDirectory);
 }
 
-
-// Display a temporary failure notification in case of an error in the shell
-void DisplayFailureNotification(const FText& InNotificationText)
-{
-	FNotificationInfo* Info = new FNotificationInfo(InNotificationText);
-	Info->ExpireDuration = 10.0f;
-	FSlateNotificationManager::Get().QueueNotification(Info);
-	// NOTE: all source control operations run in a thread, so we cannot use MessageLog nor Notify() them since they can only be used from the Main/UI thread
-	// FMessageLog("SourceControl").Error(InNotificationText);
-	UE_LOG(LogSourceControl, Error, TEXT("%s"), *InNotificationText.ToString());
-}
-
 // Internal function (called under the critical section)
 static bool _RunCommandInternal(const FString& InCommand, const TArray<FString>& InParameters, const TArray<FString>& InFiles, FString& OutResults, FString& OutErrors)
 {
@@ -175,7 +176,7 @@ static bool _RunCommandInternal(const FString& InCommand, const TArray<FString>&
 		_RestartBackgroundCommandLineShell();
 	}
 
-	// Start with the Plastic command itself ("status", "log", "checkin"...)
+	// Start with the command itself ("status", "log", "checkin"...)
 	FString FullCommand = InCommand;
 	// Append to the command all parameters, and then finally the files
 	for (const FString& Parameter : InParameters)
@@ -235,8 +236,8 @@ static bool _RunCommandInternal(const FString& InCommand, const TArray<FString>&
 			const uint32 IndexPrompt = OutResults.Find(ShellUserInteractText, ESearchCase::CaseSensitive, ESearchDir::FromEnd);
 			if (INDEX_NONE != IndexPrompt)
 			{
-				const FText ShellRequiresInteractionError(LOCTEXT("SourceControlShell_AskAuthenticate", "Plastic SCM command line requires user interaction.\nSign in using the Plastic SCM client."));
-				DisplayFailureNotification(ShellRequiresInteractionError);
+				const FText ShellRequiresInteractionError(LOCTEXT("SourceControlShell_AskAuthenticate", "Unity Version Control command line requires user interaction.\nSign in using the Unity Version Control client."));
+				FNotification::DisplayFailure(ShellRequiresInteractionError);
 
 				// Restart the shell without waiting, it is forever blocked waiting for user input
 				_RestartBackgroundCommandLineShell(true);
@@ -313,7 +314,7 @@ static bool _RunCommandInternal(const FString& InCommand, const TArray<FString>&
 	return bResult;
 }
 
-// Launch the Plastic SCM 'cm shell' process in background for optimized successive commands (thread-safe)
+// Launch the Unity Version Control 'cm shell' process in background for optimized successive commands (thread-safe)
 bool Launch(const FString& InPathToPlasticBinary, const FString& InWorkingDirectory)
 {
 	// Protect public APIs from multi-thread access
@@ -334,6 +335,20 @@ void Terminate()
 	_ExitBackgroundCommandLineShell();
 }
 
+void SetShellIsWarmedUp()
+{
+	FScopeLock Lock(&ShellCriticalSection);
+
+	_SetShellIsWarmedUp();
+}
+
+bool GetShellIsWarmedUp()
+{
+	FScopeLock Lock(&ShellCriticalSection);
+
+	return _GetShellIsWarmedUp();
+}
+
 // Run command and return the raw result
 bool RunCommand(const FString& InCommand, const TArray<FString>& InParameters, const TArray<FString>& InFiles, FString& OutResults, FString& OutErrors)
 {
diff --git a/Engine/Plugins/Developer/PlasticSourceControl/Source/PlasticSourceControl/Private/PlasticSourceControlShell.h b/Engine/Plugins/Developer/PlasticSourceControl/Source/PlasticSourceControl/Private/PlasticSourceControlShell.h
--- a/Engine/Plugins/Developer/PlasticSourceControl/Source/PlasticSourceControl/Private/PlasticSourceControlShell.h
+++ b/Engine/Plugins/Developer/PlasticSourceControl/Source/PlasticSourceControl/Private/PlasticSourceControlShell.h
@@ -3,7 +3,6 @@
 #pragma once
 
 #include "CoreMinimal.h"
-#include "ISourceControlProvider.h" // for EConcurrency
 
 namespace PlasticSourceControlShell
 {
@@ -15,7 +14,7 @@ namespace PlasticSourceControlShell
 
 
 /**
- * Launch the Plastic SCM "shell" command line process to run it in the background.
+ * Launch the Unity Version Control "shell" command line process to run it in the background.
  *
  * @param	InPathToPlasticBinary	The path to the Plastic binary
  * @param	InWorkspaceRoot			The workspace from where to run the command - usually the Game directory
@@ -26,6 +25,16 @@ bool Launch(const FString& InPathToPlasticBinary, const FString& InWorkspaceRoot
 /** Terminate the background 'cm shell' process and associated pipes */
 void Terminate();
 
+/** Mark the current shell process as already warmed up - i.e. we already ran a preliminary 'status' command. */
+void SetShellIsWarmedUp();
+
+/**
+ * Retrieve whether the current shell process was already warmed up - i.e. we already ran a preliminary 'status' command.
+ *
+ * @returns true if the shell process was already warmed up
+ */
+bool GetShellIsWarmedUp();
+
 
 /**
  * Run a Plastic command - the result is the output of cm, as a multi-line string.
diff --git a/Engine/Plugins/Developer/PlasticSourceControl/Source/PlasticSourceControl/Private/PlasticSourceControlState.cpp b/Engine/Plugins/Developer/PlasticSourceControl/Source/PlasticSourceControl/Private/PlasticSourceControlState.cpp
--- a/Engine/Plugins/Developer/PlasticSourceControl/Source/PlasticSourceControl/Private/PlasticSourceControlState.cpp
+++ b/Engine/Plugins/Developer/PlasticSourceControl/Source/PlasticSourceControl/Private/PlasticSourceControlState.cpp
@@ -16,7 +16,8 @@ const TCHAR* FPlasticSourceControlState::ToString() const
 		case EWorkspaceState::Unknown: WorkspaceStateStr = TEXT("Unknown"); break;
 		case EWorkspaceState::Ignored: WorkspaceStateStr = TEXT("Ignored"); break;
 		case EWorkspaceState::Controlled: WorkspaceStateStr = TEXT("Controlled"); break;
-		case EWorkspaceState::CheckedOut: WorkspaceStateStr = TEXT("CheckedOut"); break;
+		case EWorkspaceState::CheckedOutChanged: WorkspaceStateStr = TEXT("CheckedOutChanged"); break;
+		case EWorkspaceState::CheckedOutUnchanged: WorkspaceStateStr = TEXT("CheckedOutUnchanged"); break;
 		case EWorkspaceState::Added: WorkspaceStateStr = TEXT("Added"); break;
 		case EWorkspaceState::Moved: WorkspaceStateStr = TEXT("Moved"); break;
 		case EWorkspaceState::Copied: WorkspaceStateStr = TEXT("Copied"); break;
@@ -89,19 +90,27 @@ ISourceControlState::FResolveInfo FPlasticSourceControlState::GetResolveInfo() c
 }
 
 #if SOURCE_CONTROL_WITH_SLATE
+
 FSlateIcon FPlasticSourceControlState::GetIcon() const
 {
 	if (!IsCurrent())
 	{
 		return FSlateIcon(FRevisionControlStyleManager::GetStyleSetName(), "RevisionControl.NotAtHeadRevision");
 	}
-	else if (WorkspaceState != EWorkspaceState::CheckedOut && WorkspaceState != EWorkspaceState::Conflicted)
+
+	if (!IsCheckedOutImplementation())
 	{
 		if (IsCheckedOutOther())
 		{
 			return FSlateIcon(FRevisionControlStyleManager::GetStyleSetName(), "RevisionControl.CheckedOutByOtherUser", NAME_None, "RevisionControl.CheckedOutByOtherUserBadge");
 		}
-		else if (IsModifiedInOtherBranch())
+
+		if (IsRetainedInOtherBranch())
+		{
+			return FSlateIcon(FRevisionControlStyleManager::GetStyleSetName(), "RevisionControl.CheckedOutByOtherUserOtherBranch", NAME_None, "RevisionControl.CheckedOutByOtherUserOtherBranchBadge");
+		}
+
+		if (IsModifiedInOtherBranch())
 		{
 			return FSlateIcon(FRevisionControlStyleManager::GetStyleSetName(), "RevisionControl.ModifiedOtherBranch", NAME_None, "RevisionControl.ModifiedBadge");
 		}
@@ -109,33 +118,32 @@ FSlateIcon FPlasticSourceControlState::GetIcon() const
 
 	switch (WorkspaceState)
 	{
-	case EWorkspaceState::CheckedOut:
+	case EWorkspaceState::CheckedOutChanged:
+	case EWorkspaceState::CheckedOutUnchanged:
 	case EWorkspaceState::Replaced: // Merged (waiting for check-in)
 		return FSlateIcon(FRevisionControlStyleManager::GetStyleSetName(), "RevisionControl.CheckedOut");
-	case EWorkspaceState::Changed: // Changed but unchecked-out file custom color icon
-		return FSlateIcon(FAppStyle::GetAppStyleSetName(), "Plastic.Changed"); // custom
 	case EWorkspaceState::Added:
 	case EWorkspaceState::Copied:
 		return FSlateIcon(FRevisionControlStyleManager::GetStyleSetName(), "RevisionControl.OpenForAdd");
 	case EWorkspaceState::Moved:
 		return FSlateIcon(FRevisionControlStyleManager::GetStyleSetName(), "RevisionControl.Branched");
 	case EWorkspaceState::Deleted:
-		return FSlateIcon(FRevisionControlStyleManager::GetStyleSetName(), "RevisionControl.MarkedForDelete");
 	case EWorkspaceState::LocallyDeleted:
-		return FSlateIcon(FAppStyle::GetAppStyleSetName(), "Plastic.LocallyDeleted"); // custom
+		return FSlateIcon(FRevisionControlStyleManager::GetStyleSetName(), "RevisionControl.MarkedForDelete");
 	case EWorkspaceState::Conflicted:
 		return FSlateIcon(FRevisionControlStyleManager::GetStyleSetName(), "RevisionControl.Conflicted");
 	case EWorkspaceState::Private: // Not controlled
+	case EWorkspaceState::Changed: // Changed but unchecked-out is in a certain way not controlled
 		return FSlateIcon(FRevisionControlStyleManager::GetStyleSetName(), "RevisionControl.NotInDepot");
 	case EWorkspaceState::Ignored:
-		return FSlateIcon(FAppStyle::GetAppStyleSetName(), "Plastic.Ignored"); // custom
 	case EWorkspaceState::Unknown:
 	case EWorkspaceState::Controlled: // Unchanged (not checked out) ie no icon
 	default:
 		return FSlateIcon();
 	}
 }
-#endif //SOURCE_CONTROL_WITH_SLATE
+
+#endif // SOURCE_CONTROL_WITH_SLATE
 
 FText FPlasticSourceControlState::GetDisplayName() const
 {
@@ -146,13 +154,20 @@ FText FPlasticSourceControlState::GetDisplayName() const
 		return FText::Format(LOCTEXT("NotCurrent", "Not at the head revision CS:{0} {1} (local revision is CS:{2})"),
 			FText::AsNumber(DepotRevisionChangeset), FText::FromString(HeadUserName), FText::AsNumber(LocalRevisionChangeset, &NoCommas));
 	}
-	else if (WorkspaceState != EWorkspaceState::CheckedOut && WorkspaceState != EWorkspaceState::Conflicted)
+
+	if (!IsCheckedOutImplementation())
 	{
 		if (IsCheckedOutOther())
 		{
-			return FText::Format(LOCTEXT("CheckedOutOther", "Checked out by {0} in {1}"), FText::FromString(LockedBy), FText::FromString(LockedWhere));
+			return FText::Format(LOCTEXT("CheckedOutOther", "Checked out by {0} on {1} (in {2}) since {3}"), FText::FromString(LockedBy), FText::FromString(LockedBranch), FText::FromString(LockedWhere), FText::AsDateTime(LockedDate));
 		}
-		else if (IsModifiedInOtherBranch())
+
+		if (IsRetainedInOtherBranch())
+		{
+			return FText::Format(LOCTEXT("RetainedLock", "Retained on {0} by {1} since {2}"), FText::FromString(LockedBranch), FText::FromString(RetainedBy), FText::AsDateTime(LockedDate));
+		}
+
+		if (IsModifiedInOtherBranch())
 		{
 			return FText::Format(LOCTEXT("ModifiedOtherBranch", "Modified in {0} as CS:{1} by {2} (local revision is CS:{3})"),
 				FText::FromString(HeadBranch), FText::AsNumber(HeadChangeList, &NoCommas), FText::FromString(HeadUserName), FText::AsNumber(LocalRevisionChangeset, &NoCommas));
@@ -167,8 +182,10 @@ FText FPlasticSourceControlState::GetDisplayName() const
 		return LOCTEXT("Ignored", "Ignored");
 	case EWorkspaceState::Controlled:
 		return LOCTEXT("Controlled", "Controlled");
-	case EWorkspaceState::CheckedOut:
-		return LOCTEXT("CheckedOut", "Checked-out");
+	case EWorkspaceState::CheckedOutChanged:
+		return LOCTEXT("CheckedOutChanged", "Checked-out (changed)");
+	case EWorkspaceState::CheckedOutUnchanged:
+		return LOCTEXT("CheckedOutUnchanged", "Checked-out (unchanged)");
 	case EWorkspaceState::Added:
 		return LOCTEXT("Added", "Added");
 	case EWorkspaceState::Moved:
@@ -178,15 +195,15 @@ FText FPlasticSourceControlState::GetDisplayName() const
 	case EWorkspaceState::Replaced:
 		return LOCTEXT("Replaced", "Replaced");
 	case EWorkspaceState::Deleted:
-		return LOCTEXT("Deleted", "Deleted");
+		return LOCTEXT("Deleted", "Removed");
 	case EWorkspaceState::LocallyDeleted:
-		return LOCTEXT("LocallyDeleted", "Missing");
+		return LOCTEXT("LocallyDeleted", "Deleted locally");
 	case EWorkspaceState::Changed:
 		return LOCTEXT("Changed", "Changed");
 	case EWorkspaceState::Conflicted:
 		return LOCTEXT("Conflicted", "Conflicted");
 	case EWorkspaceState::Private:
-		return LOCTEXT("NotControlled", "Not Under Source Control");
+		return LOCTEXT("NotControlled", "Not Under Revision Control");
 	}
 
 	return FText();
@@ -201,13 +218,20 @@ FText FPlasticSourceControlState::GetDisplayTooltip() const
 		return FText::Format(LOCTEXT("NotCurrent_Tooltip", "Not at the head revision CS:{0} {1} (local revision is CS:{2})"),
 			FText::AsNumber(DepotRevisionChangeset), FText::FromString(HeadUserName), FText::AsNumber(LocalRevisionChangeset, &NoCommas));
 	}
-	else if (WorkspaceState != EWorkspaceState::CheckedOut && WorkspaceState != EWorkspaceState::Conflicted)
+
+	if (!IsCheckedOutImplementation())
 	{
 		if (IsCheckedOutOther())
 		{
-			return FText::Format(LOCTEXT("CheckedOutOther_Tooltip", "Checked out by {0} in {1}"), FText::FromString(LockedBy), FText::FromString(LockedWhere));
+			return FText::Format(LOCTEXT("CheckedOutOther_Tooltip", "Checked out by {0} on {1} (in {2}) since {3}"), FText::FromString(LockedBy), FText::FromString(LockedBranch), FText::FromString(LockedWhere), FText::AsDateTime(LockedDate));
+		}
+
+		if (IsRetainedInOtherBranch())
+		{
+			return FText::Format(LOCTEXT("RetainedLock_Tooltip", "Retained on {0} by {1} since {2}"), FText::FromString(LockedBranch), FText::FromString(RetainedBy), FText::AsDateTime(LockedDate));
 		}
-		else if (IsModifiedInOtherBranch())
+
+		if (IsModifiedInOtherBranch())
 		{
 			return FText::Format(LOCTEXT("ModifiedOtherBranch_Tooltip", "Modified in {0} as CS:{1} by {2} (local revision is CS:{3})"),
 				FText::FromString(HeadBranch), FText::AsNumber(HeadChangeList, &NoCommas), FText::FromString(HeadUserName), FText::AsNumber(LocalRevisionChangeset, &NoCommas));
@@ -222,8 +246,10 @@ FText FPlasticSourceControlState::GetDisplayTooltip() const
 		return LOCTEXT("Ignored_Tooltip", "Ignored");
 	case EWorkspaceState::Controlled:
 		return FText();
-	case EWorkspaceState::CheckedOut:
-		return LOCTEXT("CheckedOut_Tooltip", "Checked-out");
+	case EWorkspaceState::CheckedOutChanged:
+		return LOCTEXT("CheckedOut_Tooltip", "Checked-out (changed)");
+	case EWorkspaceState::CheckedOutUnchanged:
+		return LOCTEXT("CheckedOut_Tooltip", "Checked-out (unchanged)");
 	case EWorkspaceState::Added:
 		return LOCTEXT("Added_Tooltip", "Added");
 	case EWorkspaceState::Moved:
@@ -236,18 +262,18 @@ FText FPlasticSourceControlState::GetDisplayTooltip() const
 	case EWorkspaceState::Copied:
 		return LOCTEXT("Copied_Tooltip", "Copied");
 	case EWorkspaceState::Replaced:
-		return LOCTEXT("Replaced_Tooltip", "Replaced: merge conflict resolved");
+		return LOCTEXT("Replaced_Tooltip", "Replaced (merged)");
 	case EWorkspaceState::Deleted:
-		return LOCTEXT("Deleted_Tooltip", "Deleted");
+		return LOCTEXT("Deleted_Tooltip", "Removed");
 	case EWorkspaceState::LocallyDeleted:
-		return LOCTEXT("LocallyDeleted_Tooltip", "Locally Deleted");
+		return LOCTEXT("LocallyDeleted_Tooltip", "Deleted locally");
 	case EWorkspaceState::Changed:
-		return LOCTEXT("Modified_Tooltip", "Locally modified");
+		return LOCTEXT("Modified_Tooltip", "Changed locally");
 	case EWorkspaceState::Conflicted:
 		return FText::Format(LOCTEXT("Conflicted_Tooltip", "Conflict merging from source/remote CS:{0} into target/local CS:{1})"),
 			FText::FromString(PendingResolveInfo.RemoteRevision), FText::AsNumber(LocalRevisionChangeset, &NoCommas));
 	case EWorkspaceState::Private:
-		return LOCTEXT("NotControlled_Tooltip", "Private: not under version control");
+		return LOCTEXT("NotControlled_Tooltip", "Private: not under revision control");
 	}
 
 	return FText();
@@ -255,15 +281,11 @@ FText FPlasticSourceControlState::GetDisplayTooltip() const
 
 const FString& FPlasticSourceControlState::GetFilename() const
 {
-	// UE_LOG(LogSourceControl, Log, TEXT("GetFilename(%s)"), *LocalFilename);
-
 	return LocalFilename;
 }
 
 const FDateTime& FPlasticSourceControlState::GetTimeStamp() const
 {
-	// UE_LOG(LogSourceControl, Log, TEXT("GetTimeStamp(%s)=%s"), *LocalFilename, *TimeStamp.ToString());
-
 	return TimeStamp;
 }
 
@@ -277,7 +299,7 @@ bool FPlasticSourceControlState::CanCheckIn() const
 							|| WorkspaceState == EWorkspaceState::Moved
 							|| WorkspaceState == EWorkspaceState::Copied
 							|| WorkspaceState == EWorkspaceState::Replaced
-							|| WorkspaceState == EWorkspaceState::CheckedOut)
+							|| WorkspaceState == EWorkspaceState::CheckedOutChanged)
 							&& !IsCheckedOutOther()	// Is not already checked-out elsewhere
 							&& IsCurrent();			// Is up to date (at the revision of the repo)
 
@@ -297,8 +319,7 @@ bool FPlasticSourceControlState::CanCheckout() const
 	}
 
 	const bool bCanCheckout  =    (WorkspaceState == EWorkspaceState::Controlled	// In source control, Unmodified
-								|| WorkspaceState == EWorkspaceState::Changed		// In source control, but not checked-out
-								|| WorkspaceState == EWorkspaceState::Replaced)		// In source control, merged, waiting for checkin to conclude the merge
+								|| WorkspaceState == EWorkspaceState::Changed)		// In source control, but not checked-out
 								&& !IsCheckedOutOther()	// Is not already checked-out elsewhere
 								&& IsCurrent();			// Is up to date (at the revision of the repo)
 
@@ -312,10 +333,7 @@ bool FPlasticSourceControlState::CanCheckout() const
 
 bool FPlasticSourceControlState::IsCheckedOut() const
 {
-	const bool bIsCheckedOut = WorkspaceState == EWorkspaceState::CheckedOut
-							|| WorkspaceState == EWorkspaceState::Moved
-							|| WorkspaceState == EWorkspaceState::Conflicted	// In source control, waiting for merged
-							|| WorkspaceState == EWorkspaceState::Replaced		// In source control, merged, waiting for checkin to conclude the merge
+	const bool bIsCheckedOut = IsCheckedOutImplementation()
 							|| WorkspaceState == EWorkspaceState::Changed;		// Note: Workaround to enable checkin (still required by UE5.0)
 
 	if (bIsCheckedOut)
@@ -334,6 +352,30 @@ bool FPlasticSourceControlState::IsCheckedOut() const
 	}
 }
 
+bool FPlasticSourceControlState::IsCheckedOutImplementation() const
+{
+	const bool bIsCheckedOut = WorkspaceState == EWorkspaceState::CheckedOutChanged
+							|| WorkspaceState == EWorkspaceState::CheckedOutUnchanged
+							|| WorkspaceState == EWorkspaceState::Added
+							|| WorkspaceState == EWorkspaceState::Deleted
+							|| WorkspaceState == EWorkspaceState::Copied
+							|| WorkspaceState == EWorkspaceState::Moved
+							|| WorkspaceState == EWorkspaceState::Conflicted	// In source control, waiting for merge
+							|| WorkspaceState == EWorkspaceState::Replaced;		// In source control, merged, waiting for checkin to conclude the merge
+
+	return bIsCheckedOut;
+}
+
+bool FPlasticSourceControlState::IsLocked() const
+{
+	return !LockedBy.IsEmpty();
+}
+
+bool FPlasticSourceControlState::IsRetainedInOtherBranch() const
+{
+	return !RetainedBy.IsEmpty();
+}
+
 bool FPlasticSourceControlState::IsCheckedOutOther(FString* Who) const
 {
 	if (Who != NULL)
@@ -341,14 +383,8 @@ bool FPlasticSourceControlState::IsCheckedOutOther(FString* Who) const
 		*Who = LockedBy;
 	}
 
-	// If the asset is Locked but not CheckedOut locally, it means it is locked somewhere else
-	const bool bIsLocked = !LockedBy.IsEmpty();
-	const bool bIsCheckedOut = WorkspaceState == EWorkspaceState::CheckedOut
-							|| WorkspaceState == EWorkspaceState::Added
-							|| WorkspaceState == EWorkspaceState::Moved
-							|| WorkspaceState == EWorkspaceState::Conflicted	// In source control, waiting for merged
-							|| WorkspaceState == EWorkspaceState::Replaced;		// In source control, merged, waiting for checkin to conclude the merge
-	const bool bIsLockedByOther = bIsLocked && !bIsCheckedOut;
+	// An asset is locked somewhere else if it is Locked but not CheckedOut on the current workspace
+	const bool bIsLockedByOther = IsLocked() && !IsCheckedOutImplementation();
 
 	if (bIsLockedByOther)
 	{
@@ -358,16 +394,15 @@ bool FPlasticSourceControlState::IsCheckedOutOther(FString* Who) const
 	return bIsLockedByOther;
 }
 
-
-/** Get whether this file is checked out in a different branch, if no branch is specified defaults to FEngineVerion current branch */
+/** Get whether this file is checked out in a different branch */
 bool FPlasticSourceControlState::IsCheckedOutInOtherBranch(const FString& CurrentBranch /* = FString() */) const
 {
-	// Note: to my knowledge, it's not possible to detect that with PlasticSCM without the Locks,
-	// which are already detected by fileinfo LockedBy/LockedWhere and reported by IsCheckedOutOther() above
-	return false;
+	// NOTE: technically this scenario isn't currently possible with Unity Version Control,
+	//       but the plugin need to use an existing Engine hook, so it's using this one as a way to display "Retained" locks
+	return IsRetainedInOtherBranch();
 }
 
-/** Get whether this file is modified in a different branch, if no branch is specified defaults to FEngineVerion current branch */
+/** Get whether this file is modified in a different branch */
 bool FPlasticSourceControlState::IsModifiedInOtherBranch(const FString& CurrentBranch /* = FString() */) const
 {
 	return !HeadBranch.IsEmpty();
@@ -403,9 +438,9 @@ bool FPlasticSourceControlState::IsSourceControlled() const
 	// NOTE: the Editor Collections rely on the default 'Unknown' state (until the actual file status is obtained) to be considered "in source control"
 	const bool bIsSourceControlled = WorkspaceState != EWorkspaceState::Private
 								  && WorkspaceState != EWorkspaceState::Ignored;
-								  // WorkspaceState != EWorkspaceState::Unknown 
+								  // WorkspaceState != EWorkspaceState::Unknown
 
-	if (!bIsSourceControlled && !IsUnknown()) UE_LOG(LogSourceControl, Log, TEXT("%s NOT SourceControlled"), *LocalFilename);
+	if (!bIsSourceControlled && !IsUnknown()) UE_LOG(LogSourceControl, Verbose, TEXT("%s NOT SourceControlled"), *LocalFilename);
 
 	return bIsSourceControlled;
 }
@@ -415,7 +450,7 @@ bool FPlasticSourceControlState::IsAdded() const
 	const bool bIsAdded =	WorkspaceState == EWorkspaceState::Added
 						 || WorkspaceState == EWorkspaceState::Copied;
 
-	if (bIsAdded) UE_LOG(LogSourceControl, Log, TEXT("%s IsAdded"), *LocalFilename);
+	if (bIsAdded) UE_LOG(LogSourceControl, Verbose, TEXT("%s IsAdded"), *LocalFilename);
 
 	return bIsAdded;
 }
@@ -425,27 +460,28 @@ bool FPlasticSourceControlState::IsDeleted() const
 	const bool bIsDeleted =	WorkspaceState == EWorkspaceState::Deleted
 						 || WorkspaceState == EWorkspaceState::LocallyDeleted;
 
-	if (bIsDeleted) UE_LOG(LogSourceControl, Log, TEXT("%s IsDeleted"), *LocalFilename);
+	if (bIsDeleted) UE_LOG(LogSourceControl, Verbose, TEXT("%s IsDeleted"), *LocalFilename);
 
 	return bIsDeleted;
 }
 
 bool FPlasticSourceControlState::IsIgnored() const
 {
-	if (WorkspaceState == EWorkspaceState::Ignored) UE_LOG(LogSourceControl, Log, TEXT("%s IsIgnored"), *LocalFilename);
+	if (WorkspaceState == EWorkspaceState::Ignored) UE_LOG(LogSourceControl, Verbose, TEXT("%s IsIgnored"), *LocalFilename);
 
 	return WorkspaceState == EWorkspaceState::Ignored;
 }
 
 bool FPlasticSourceControlState::CanEdit() const
 {
-	const bool bCanEdit =  WorkspaceState == EWorkspaceState::CheckedOut
+	const bool bCanEdit =  WorkspaceState == EWorkspaceState::CheckedOutChanged
+						|| WorkspaceState == EWorkspaceState::CheckedOutUnchanged
 						|| WorkspaceState == EWorkspaceState::Added
 						|| WorkspaceState == EWorkspaceState::Moved
 						|| WorkspaceState == EWorkspaceState::Copied
 						|| WorkspaceState == EWorkspaceState::Replaced;
 
-	UE_LOG(LogSourceControl, Log, TEXT("%s CanEdit=%d"), *LocalFilename, bCanEdit);
+	UE_LOG(LogSourceControl, Verbose, TEXT("%s CanEdit=%d"), *LocalFilename, bCanEdit);
 
 	return bCanEdit;
 }
@@ -462,16 +498,17 @@ bool FPlasticSourceControlState::IsUnknown() const
 
 bool FPlasticSourceControlState::IsModified() const
 {
-	// Warning: for a clean "checkin" (commit) checked-out files unmodified should be removed from the changeset (the index)
+	// Warning: for a clean "checkin" (commit) checked-out files unmodified should be removed from the changeset (Perforce)
 	//
 	// Thus, before checkin UE4 Editor call RevertUnchangedFiles() in PromptForCheckin() and CheckinFiles().
 	//
 	// So here we must take care to enumerate all states that need to be committed, all other will be discarded:
 	//  - Unknown
 	//  - Controlled (Unchanged)
+	//  - CheckedOutUnchanged
 	//  - Private (Not Controlled)
 	//  - Ignored
-	const bool bIsModified =   WorkspaceState == EWorkspaceState::CheckedOut
+	const bool bIsModified =   WorkspaceState == EWorkspaceState::CheckedOutChanged
 							|| WorkspaceState == EWorkspaceState::Added
 							|| WorkspaceState == EWorkspaceState::Moved
 							|| WorkspaceState == EWorkspaceState::Copied
@@ -489,7 +526,7 @@ bool FPlasticSourceControlState::IsModified() const
 
 bool FPlasticSourceControlState::CanAdd() const
 {
-	if (!IsUnknown()) UE_LOG(LogSourceControl, Log, TEXT("%s CanAdd=%d"), *LocalFilename, WorkspaceState == EWorkspaceState::Private);
+	if (!IsUnknown()) UE_LOG(LogSourceControl, Verbose, TEXT("%s CanAdd=%d"), *LocalFilename, WorkspaceState == EWorkspaceState::Private);
 
 	return WorkspaceState == EWorkspaceState::Private;
 }
@@ -503,7 +540,7 @@ bool FPlasticSourceControlState::IsConflicted() const
 
 bool FPlasticSourceControlState::CanRevert() const
 {
-	return IsModified();
+	return IsModified() || WorkspaceState == EWorkspaceState::CheckedOutUnchanged;
 }
 
 #undef LOCTEXT_NAMESPACE
diff --git a/Engine/Plugins/Developer/PlasticSourceControl/Source/PlasticSourceControl/Private/PlasticSourceControlState.h b/Engine/Plugins/Developer/PlasticSourceControl/Source/PlasticSourceControl/Private/PlasticSourceControlState.h
--- a/Engine/Plugins/Developer/PlasticSourceControl/Source/PlasticSourceControl/Private/PlasticSourceControlState.h
+++ b/Engine/Plugins/Developer/PlasticSourceControl/Source/PlasticSourceControl/Private/PlasticSourceControlState.h
@@ -6,23 +6,25 @@
 #include "ISourceControlState.h"
 #include "ISourceControlRevision.h"
 #include "PlasticSourceControlRevision.h"
+
 #include "PlasticSourceControlChangelist.h"
 
 enum class EWorkspaceState
 {
-		Unknown,
-		Ignored,
-		Controlled, // called "Pristine" in Perforce, "Unchanged" in Git, "Clean" in SVN
-		CheckedOut,
-		Added,
-		Moved, // Renamed
-		Copied,
-		Replaced, // Replaced / Merged
-		Deleted,
-		LocallyDeleted, // Missing
-		Changed, // Modified but not CheckedOut
-		Conflicted,
-		Private, // "Not Controlled"/"Not In Depot"/"Untracked"
+	Unknown,
+	Ignored,
+	Controlled, // called "Pristine" in Perforce, "Unchanged" in Git, "Clean" in SVN
+	CheckedOutChanged, // Checked-out, with changes (or without knowing for older version of Unity Version Control)
+	CheckedOutUnchanged, // Checked-out with no changes (cannot be checked-in and can be reverted by UndoUnchanged)
+	Added,
+	Moved, // Renamed
+	Copied,
+	Replaced, // Replaced / Merged
+	Deleted,
+	LocallyDeleted, // Missing
+	Changed, // Locally Changed but not CheckedOut
+	Conflicted,
+	Private, // "Not Controlled"/"Not In Depot"/"Untracked"
 };
 
 class FPlasticSourceControlState : public ISourceControlState
@@ -67,6 +69,10 @@ public:
 		{
 			LockedBy = MoveTemp(InState.LockedBy);
 			LockedWhere = MoveTemp(InState.LockedWhere);
+			LockedBranch = MoveTemp(InState.LockedBranch);
+			LockedId = MoveTemp(InState.LockedId);
+			LockedDate = MoveTemp(InState.LockedDate);
+			RetainedBy = MoveTemp(InState.RetainedBy);
 			RepSpec = MoveTemp(InState.RepSpec);
 			DepotRevisionChangeset = InState.DepotRevisionChangeset;
 			LocalRevisionChangeset = InState.LocalRevisionChangeset;
@@ -99,7 +105,7 @@ public:
 	virtual FResolveInfo GetResolveInfo() const override;
 #if SOURCE_CONTROL_WITH_SLATE
 	virtual FSlateIcon GetIcon() const override;
-#endif //SOURCE_CONTROL_WITH_SLATE
+#endif // SOURCE_CONTROL_WITH_SLATE
 	virtual FText GetDisplayName() const override;
 	virtual FText GetDisplayTooltip() const override;
 	virtual const FString& GetFilename() const override;
@@ -127,6 +133,10 @@ public:
 	virtual bool IsConflicted() const override;
 	virtual bool CanRevert() const override;
 
+	bool IsCheckedOutImplementation() const;
+	bool IsLocked() const;
+	bool IsRetainedInOtherBranch() const;
+
 public:
 	/** History of the item, if any */
 	TPlasticSourceControlHistory History;
@@ -140,32 +150,35 @@ public:
 	/** Pending rev info with which a file must be resolved, invalid if no resolve pending */
 	FResolveInfo PendingResolveInfo;
 
-	UE_DEPRECATED(5.3, "Use PendingResolveInfo.BaseFile instead")
-	FString PendingMergeFilename;
-
-	UE_DEPRECATED(5.3, "Use PendingResolveInfo.BaseRevision instead")
-	int32 PendingMergeBaseChangeset = INVALID_REVISION;
-
-	UE_DEPRECATED(5.3, "Use PendingResolveInfo.RemoteRevision instead")
-	int32 PendingMergeSourceChangeset = INVALID_REVISION;
-
-	/** Plastic SCM Parameters of the merge in progress */
+	/** Unity Version Control Parameters of the merge in progress */
 	TArray<FString> PendingMergeParameters;
 
 	/** If a user (another or ourself) has this file locked, this contains their name. */
 	FString LockedBy;
 
-	/** Location of the locked file. */
+	/** Location (Workspace) where the file was exclusively checked-out. */
 	FString LockedWhere;
 
+	/** Branch where the file was Locked or is Retained. */
+	FString LockedBranch;
+
+	/** Item id of the locked file (for an admin to unlock it). */
+	int32 LockedId = INVALID_REVISION;
+
+	/** Date when the file was Locked. */
+	FDateTime LockedDate = 0;
+
+	/** If a user (another or ourself) has this file Retained on another branch, this contains their name. */
+	FString RetainedBy;
+
 	/** State of the workspace */
 	EWorkspaceState WorkspaceState = EWorkspaceState::Unknown;
 
 	/** Latest revision number of the file in the depot (on the current branch) */
-	int DepotRevisionChangeset = INVALID_REVISION;
+	int32 DepotRevisionChangeset = INVALID_REVISION;
 
 	/** Latest revision number at which a file was synced to before being edited */
-	int LocalRevisionChangeset = INVALID_REVISION;
+	int32 LocalRevisionChangeset = INVALID_REVISION;
 
 	/** Original name in case of a Moved/Renamed file */
 	FString MovedFrom;
diff --git a/Engine/Plugins/Developer/PlasticSourceControl/Source/PlasticSourceControl/Private/PlasticSourceControlStyle.cpp b/Engine/Plugins/Developer/PlasticSourceControl/Source/PlasticSourceControl/Private/PlasticSourceControlStyle.cpp
new file mode 100644
--- /dev/null
+++ b/Engine/Plugins/Developer/PlasticSourceControl/Source/PlasticSourceControl/Private/PlasticSourceControlStyle.cpp
@@ -0,0 +1,69 @@
+// Copyright Epic Games, Inc. All Rights Reserved.
+
+#include "PlasticSourceControlStyle.h"
+
+#include "Framework/Application/SlateApplication.h"
+#include "Interfaces/IPluginManager.h"
+#include "Slate/SlateGameResources.h"
+#include "Styling/SlateStyle.h"
+#include "Styling/SlateStyleMacros.h"
+#include "Styling/SlateStyleRegistry.h"
+
+#include "PlasticSourceControlModule.h"
+
+TSharedPtr<FSlateStyleSet> FPlasticSourceControlStyle::StyleInstance = nullptr;
+
+void FPlasticSourceControlStyle::Initialize()
+{
+	if (!StyleInstance.IsValid())
+	{
+		StyleInstance = Create();
+		FSlateStyleRegistry::RegisterSlateStyle(*StyleInstance);
+	}
+}
+
+void FPlasticSourceControlStyle::Shutdown()
+{
+	FSlateStyleRegistry::UnRegisterSlateStyle(*StyleInstance);
+	ensure(StyleInstance.IsUnique());
+	StyleInstance.Reset();
+}
+
+FName FPlasticSourceControlStyle::GetStyleSetName()
+{
+	static FName StyleSetName(TEXT("PlasticSourceControlStyle"));
+	return StyleSetName;
+}
+
+const FVector2D Icon16x16(16.0f, 16.0f);
+
+TSharedRef<FSlateStyleSet> FPlasticSourceControlStyle::Create()
+{
+	TSharedRef<FSlateStyleSet> Style = MakeShareable(new FSlateStyleSet("PlasticSourceControlStyle"));
+	Style->SetContentRoot(FPlasticSourceControlModule::GetPlugin()->GetBaseDir() / TEXT("Resources"));
+
+	Style->Set("PlasticSourceControl.PluginIcon.Small", new FSlateImageBrush(FPlasticSourceControlStyle::InContent("Icon128", ".png"), Icon16x16));
+
+	return Style;
+}
+
+FString FPlasticSourceControlStyle::InContent(const FString& RelativePath, const ANSICHAR* Extension)
+{
+	auto myself = FPlasticSourceControlModule::GetPlugin();
+	check(myself.IsValid());
+	static FString ContentDir = myself->GetBaseDir() / TEXT("Resources");
+	return (ContentDir / RelativePath) + Extension;
+}
+
+void FPlasticSourceControlStyle::ReloadTextures()
+{
+	if (FSlateApplication::IsInitialized())
+	{
+		FSlateApplication::Get().GetRenderer()->ReloadTextureResources();
+	}
+}
+
+const ISlateStyle& FPlasticSourceControlStyle::Get()
+{
+	return *StyleInstance;
+}
\ No newline at end of file
diff --git a/Engine/Plugins/Developer/PlasticSourceControl/Source/PlasticSourceControl/Private/PlasticSourceControlStyle.h b/Engine/Plugins/Developer/PlasticSourceControl/Source/PlasticSourceControl/Private/PlasticSourceControlStyle.h
new file mode 100644
--- /dev/null
+++ b/Engine/Plugins/Developer/PlasticSourceControl/Source/PlasticSourceControl/Private/PlasticSourceControlStyle.h
@@ -0,0 +1,33 @@
+// Copyright Epic Games, Inc. All Rights Reserved.
+
+#pragma once
+
+#include "CoreMinimal.h"
+
+class FSlateStyleSet;
+class ISlateStyle;
+
+class FPlasticSourceControlStyle
+{
+public:
+	static void Initialize();
+
+	static void Shutdown();
+
+	/** reloads textures used by slate renderer */
+	static void ReloadTextures();
+
+	/** @return The Slate style set */
+	static const ISlateStyle& Get();
+
+	static FName GetStyleSetName();
+
+private:
+	static TSharedRef<FSlateStyleSet> Create();
+
+private:
+	static FString InContent(const FString& RelativePath, const ANSICHAR* Extension);
+
+private:
+	static TSharedPtr<class FSlateStyleSet> StyleInstance;
+};
\ No newline at end of file
