PR #11295: [Source Control] Update Plastic SCM provider to 1.9.0 (from 1.8.0)

diff --git a/Engine/Plugins/Developer/PlasticSourceControl/Source/PlasticSourceControl/Private/PlasticSourceControlUtils.cpp b/Engine/Plugins/Developer/PlasticSourceControl/Source/PlasticSourceControl/Private/PlasticSourceControlUtils.cpp
--- a/Engine/Plugins/Developer/PlasticSourceControl/Source/PlasticSourceControl/Private/PlasticSourceControlUtils.cpp
+++ b/Engine/Plugins/Developer/PlasticSourceControl/Source/PlasticSourceControl/Private/PlasticSourceControlUtils.cpp
@@ -2,8 +2,10 @@
 
 #include "PlasticSourceControlUtils.h"
 
+#include "PlasticSourceControlBranch.h"
 #include "PlasticSourceControlCommand.h"
 #include "PlasticSourceControlModule.h"
+#include "PlasticSourceControlParsers.h"
 #include "PlasticSourceControlProjectSettings.h"
 #include "PlasticSourceControlProvider.h"
 #include "PlasticSourceControlSettings.h"
@@ -12,27 +14,25 @@
 #include "PlasticSourceControlVersions.h"
 #include "ISourceControlModule.h"
 
-#include "HAL/PlatformFileManager.h"
-#include "HAL/PlatformProcess.h"
 #include "Misc/FileHelper.h"
 #include "Misc/Paths.h"
-#include "XmlParser.h"
 #include "SoftwareVersion.h"
 #include "ScopedTempFile.h"
 
-#include "PlasticSourceControlChangelist.h"
 #include "PlasticSourceControlChangelistState.h"
 
 namespace PlasticSourceControlUtils
 {
 
+#define FILE_STATUS_SEPARATOR TEXT(";")
+
 // Run a command and return the result as raw strings
 bool RunCommand(const FString& InCommand, const TArray<FString>& InParameters, const TArray<FString>& InFiles, FString& OutResults, FString& OutErrors)
 {
 	return PlasticSourceControlShell::RunCommand(InCommand, InParameters, InFiles, OutResults, OutErrors);
 }
 
-// Run a command with basic parsing or results & errors from the Plastic command line process
+// Run a command with basic parsing or results & errors from the cm command line process
 bool RunCommand(const FString& InCommand, const TArray<FString>& InParameters, const TArray<FString>& InFiles, TArray<FString>& OutResults, TArray<FString>& OutErrorMessages)
 {
 	FString Results;
@@ -46,11 +46,6 @@ bool RunCommand(const FString& InCommand, const TArray<FString>& InParameters, c
 	return bResult;
 }
 
-const FSoftwareVersion& GetOldestSupportedPlasticScmVersion()
-{
-	return PlasticSourceControlVersions::OldestSupported;
-}
-
 FString FindPlasticBinaryPath()
 {
 #if PLATFORM_WINDOWS
@@ -60,7 +55,7 @@ FString FindPlasticBinaryPath()
 #endif
 }
 
-// Find the root of the Plastic workspace, looking from the provided path and upward in its parent directories.
+// Find the root of the workspace, looking from the provided path and upward in its parent directories.
 bool GetWorkspacePath(const FString& InPath, FString& OutWorkspaceRoot)
 {
 	TArray<FString> Results;
@@ -126,6 +121,21 @@ bool GetConfigSetFilesAsReadOnly()
 	return false;
 }
 
+FString GetConfigDefaultRepServer()
+{
+	FString ServerUrl;
+	TArray<FString> Results;
+	TArray<FString> ErrorMessages;
+	TArray<FString> Parameters;
+	Parameters.Add(TEXT("defaultrepserver"));
+	const bool bResult = RunCommand(TEXT("getconfig"), Parameters, TArray<FString>(), Results, ErrorMessages);
+	if (bResult && Results.Num() > 0)
+	{
+		ServerUrl = MoveTemp(Results[0]);
+	}
+	return ServerUrl;
+}
+
 void GetUserName(FString& OutUserName)
 {
 	TArray<FString> Results;
@@ -158,117 +168,26 @@ bool GetWorkspaceName(const FString& InWorkspaceRoot, FString& OutWorkspaceName,
 	return bResult;
 }
 
-static bool ParseWorkspaceInformation(const TArray<FString>& InResults, int32& OutChangeset, FString& OutRepositoryName, FString& OutServerUrl, FString& OutBranchName)
-{
-	bool bResult = true;
-
-	// Get workspace status, in the form "cs:41@rep:UEPlasticPlugin@repserver:localhost:8087" (disabled by the "--nostatus" flag)
-	//                                or "cs:41@rep:UEPlasticPlugin@repserver:SRombauts@cloud" (when connected directly to the cloud)
-	if (InResults.Num() > 0)
-	{
-		static const FString ChangesetPrefix(TEXT("cs:"));
-		static const FString RepPrefix(TEXT("@rep:"));
-		static const FString ServerPrefix(TEXT("@repserver:"));
-		const FString& WorkspaceStatus = InResults[0];
-		const int32 RepIndex = WorkspaceStatus.Find(RepPrefix, ESearchCase::CaseSensitive);
-		const int32 ServerIndex = WorkspaceStatus.Find(ServerPrefix, ESearchCase::CaseSensitive, ESearchDir::FromEnd);
-		if ((RepIndex > INDEX_NONE) && (ServerIndex > INDEX_NONE))
-		{
-			const FString ChangesetString = WorkspaceStatus.Mid(ChangesetPrefix.Len(), RepIndex - ChangesetPrefix.Len());
-			OutChangeset = FCString::Atoi(*ChangesetString);
-			OutRepositoryName = WorkspaceStatus.Mid(RepIndex + RepPrefix.Len(), ServerIndex - RepIndex - RepPrefix.Len());
-			OutServerUrl = WorkspaceStatus.RightChop(ServerIndex + ServerPrefix.Len());
-		}
-		else
-		{
-			bResult = false;
-		}
-	}
-	// Get the branch name, in the form "Branch /main@UE5PlasticPluginDev@test@cloud" (enabled by the "--wkconfig" flag)
-	if (InResults.Num() > 1)
-	{
-		static const FString BranchPrefix(TEXT("Branch "));
-		const FString& BranchInfo = InResults[1];
-		const int32 BranchIndex = BranchInfo.Find(BranchPrefix, ESearchCase::CaseSensitive);
-		if (BranchIndex > INDEX_NONE)
-		{
-			OutBranchName = BranchInfo;
-		}
-	}
-
-	return bResult;
-}
-
-bool GetWorkspaceInformation(int32& OutChangeset, FString& OutRepositoryName, FString& OutServerUrl, FString& OutBranchName, TArray<FString>& OutErrorMessages)
+bool GetWorkspaceInfo(FString& OutBranchName, FString& OutRepositoryName, FString& OutServerUrl, TArray<FString>& OutErrorMessages)
 {
 	TArray<FString> Results;
-	TArray<FString> Parameters;
-	Parameters.Add(TEXT("--compact"));
-	Parameters.Add(TEXT("--header")); // Only prints the workspace status. No file status.
-	// NOTE: --wkconfig results in two network calls GetBranchInfoByName & GetLastChangesetOnBranch so it's okay to do it once here but not all the time
-	Parameters.Add(TEXT("--wkconfig")); // Branch name. NOTE: Deprecated in 8.0.16.3000 https://www.plasticscm.com/download/releasenotes/8.0.16.3000
-	bool bResult = RunCommand(TEXT("status"), Parameters, TArray<FString>(), Results, OutErrorMessages);
+	bool bResult = RunCommand(TEXT("workspaceinfo"), TArray<FString>(), TArray<FString>(), Results, OutErrorMessages);
 	if (bResult)
 	{
-		bResult = ParseWorkspaceInformation(Results, OutChangeset, OutRepositoryName, OutServerUrl, OutBranchName);
+		bResult = PlasticSourceControlParsers::ParseWorkspaceInfo(Results, OutBranchName, OutRepositoryName, OutServerUrl);
 	}
 
 	return bResult;
 }
 
-static bool ParseWorkspaceInfo(const TArray<FString>& InResults, FString& OutBranchName, FString& OutRepositoryName, FString& OutServerUrl)
-{
-	if (InResults.Num() == 0)
-	{
-		return false;
-	}
-
-	// Get workspace information, in the form "Branch /main@UE5PlasticPluginDev@localhost:8087"
-	//                                     or "Branch /main@UE5PlasticPluginDev@test@cloud" (when connected directly to the cloud)
-	//                                     or "Branch /main@rep:UE5OpenWorldPerfTest@repserver:test@cloud"
-	static const FString BranchPrefix(TEXT("Branch "));
-	if (!InResults[0].StartsWith(BranchPrefix, ESearchCase::CaseSensitive))
-	{
-		return false;
-	}
-	FString Temp = InResults[0].RightChop(BranchPrefix.Len());
-	int32 SeparatorIndex;
-	if (!Temp.FindChar(TEXT('@'), SeparatorIndex))
-	{
-		return false;
-	}
-	OutBranchName = Temp.Left(SeparatorIndex);
-	Temp.RightChopInline(SeparatorIndex + 1);
-	if (!Temp.FindChar(TEXT('@'), SeparatorIndex))
-	{
-		return false;
-	}
-	OutRepositoryName = Temp.Left(SeparatorIndex);
-	static const FString RepPrefix(TEXT("rep:"));
-	if (OutRepositoryName.StartsWith(RepPrefix, ESearchCase::CaseSensitive))
-	{
-		OutRepositoryName.RightChopInline(RepPrefix.Len());
-	}
-	OutServerUrl = Temp.RightChop(SeparatorIndex + 1);
-	static const FString RepserverPrefix(TEXT("repserver:"));
-	if (OutServerUrl.StartsWith(RepserverPrefix, ESearchCase::CaseSensitive))
-	{
-		OutServerUrl.RightChopInline(RepserverPrefix.Len());
-	}
-
-	return true;
-}
-
-bool GetWorkspaceInfo(FString& OutRepositoryName, FString& OutServerUrl, FString& OutBranchName, TArray<FString>& OutErrorMessages)
+bool RunCheckConnection(FString& OutBranchName, FString& OutRepositoryName, FString& OutServerUrl, TArray<FString>& OutInfoMessages, TArray<FString>& OutErrorMessages)
 {
-	TArray<FString> Results;
-	bool bResult = RunCommand(TEXT("workspaceinfo"), TArray<FString>(), TArray<FString>(), Results, OutErrorMessages);
-	if (bResult)
+	TArray<FString> Parameters;
+	if (PlasticSourceControlUtils::GetWorkspaceInfo(OutBranchName, OutRepositoryName, OutServerUrl, OutErrorMessages))
 	{
-		bResult = ParseWorkspaceInfo(Results, OutRepositoryName, OutServerUrl, OutBranchName);
+		Parameters.Add(FString::Printf(TEXT("--server=%s"), *OutServerUrl));
 	}
-
-	return bResult;
+	return PlasticSourceControlUtils::RunCommand(TEXT("checkconnection"), Parameters, TArray<FString>(), OutInfoMessages, OutErrorMessages);
 }
 
 FString UserNameToDisplayName(const FString& InUserName)
@@ -290,319 +209,55 @@ FString UserNameToDisplayName(const FString& InUserName)
 }
 
 /**
- * @brief Extract the renamed from filename from a Plastic SCM status result.
- *
- * Examples of status results:
- MV 100% Content\ToMove_BP.uasset -> Content\Moved_BP.uasset
- *
- * @param[in] InResult One line of status
- * @return Renamed from filename extracted from the line of status
- *
- * @see FilenameFromPlasticStatus()
- */
-static FString RenamedFromPlasticStatus(const FString& InResult)
-{
-	FString RenamedFrom;
-	int32 RenameIndex;
-	if (InResult.FindLastChar(TEXT('>'), RenameIndex))
-	{
-		// Extract only the first part of a rename "from -> to" (after the 2 letters status surrounded by 2 spaces)
-		RenamedFrom = InResult.Mid(9, RenameIndex - 9 - 2);
-	}
-	return RenamedFrom;
-}
-
-/**
- * @brief Extract the relative filename from a Plastic SCM status result.
- *
- * Examples of status results:
- CO Content\CheckedOut_BP.uasset
- CO Content\Merged_BP.uasset (Merge from 140)
- MV 100% Content\ToMove_BP.uasset -> Content\Moved_BP.uasset
- *
- * @param[in] InResult One line of status
- * @return Relative filename extracted from the line of status
- *
- * @see FPlasticStatusFileMatcher and StateFromPlasticStatus()
- */
-static FString FilenameFromPlasticStatus(const FString& InResult)
-{
-	FString RelativeFilename;
-	int32 RenameIndex;
-	if (InResult.FindLastChar(TEXT('>'), RenameIndex))
-	{
-		// Extract only the second part of a rename "from -> to"
-		RelativeFilename = InResult.RightChop(RenameIndex + 2);
-	}
-	else
-	{
-		// Extract the relative filename from the Plastic SCM status result (after the 2 letters status surrounded by 2 spaces)
-		RelativeFilename = InResult.RightChop(4);
-
-		const int32 MergeIndex = RelativeFilename.Find(TEXT(" (Merge from "));
-		if (MergeIndex != INDEX_NONE)
-		{
-			RelativeFilename.LeftInline(MergeIndex);
-		}
-	}
-
-	return RelativeFilename;
-}
-
-/**
- * @brief Match the relative filename of a Plastic SCM status result with a provided absolute filename
- *
- * Examples of status results:
- CO Content\CheckedOut_BP.uasset
- MV 100% Content\ToMove_BP.uasset -> Content\Moved_BP.uasset
- */
-class FPlasticStatusFileMatcher
-{
-public:
-	explicit FPlasticStatusFileMatcher(const FString& InAbsoluteFilename)
-		: AbsoluteFilename(InAbsoluteFilename)
-	{
-	}
-
-	bool operator()(const FString& InResult) const
-	{
-		return AbsoluteFilename.Contains(FilenameFromPlasticStatus(InResult));
-	}
-
-private:
-	const FString& AbsoluteFilename;
-};
-
-/**
- * Extract and interpret the file state from the given Plastic "status" result.
- * empty string = unmodified/controlled or hidden changes
- CH Content\Changed_BP.uasset
- CO Content\CheckedOut_BP.uasset
- CP Content\Copied_BP.uasset
- RP Content\Replaced_BP.uasset
- AD Content\Added_BP.uasset
- PR Content\Private_BP.uasset
- IG Content\Ignored_BP.uasset
- DE Content\Deleted_BP.uasset
- LD Content\Deleted2_BP.uasset
- MV 100% Content\ToMove_BP.uasset -> Content\Moved_BP.uasset
- LM 100% Content\ToMove2_BP.uasset -> Content\Moved2_BP.uasset
- */
-static EWorkspaceState StateFromPlasticStatus(const FString& InResult)
-{
-	EWorkspaceState State;
-	const FString FileStatus = (InResult[0] == TEXT(' ')) ? InResult.Mid(1, 2) : InResult;
-
-	if (FileStatus == "CH") // Modified but not Checked-Out
-	{
-		State = EWorkspaceState::Changed;
-	}
-	else if (FileStatus == "CO") // Checked-Out for modification
-	{
-		State = EWorkspaceState::CheckedOut;
-	}
-	else if (FileStatus == "CP")
-	{
-		State = EWorkspaceState::Copied;
-	}
-	else if (FileStatus == "RP")
-	{
-		State = EWorkspaceState::Replaced;
-	}
-	else if (FileStatus == "AD")
-	{
-		State = EWorkspaceState::Added;
-	}
-	else if ((FileStatus == "PR") || (FileStatus == "LM")) // Not Controlled/Not in Depot/Untracked (or Locally Moved/Renamed)
-	{
-		State = EWorkspaceState::Private;
-	}
-	else if (FileStatus == "IG")
-	{
-		State = EWorkspaceState::Ignored;
-	}
-	else if (FileStatus == "DE")
-	{
-		State = EWorkspaceState::Deleted; // Deleted (removed from source control)
-	}
-	else if (FileStatus == "LD")
-	{
-		State = EWorkspaceState::LocallyDeleted; // Locally Deleted (ie. missing)
-	}
-	else if (FileStatus == "MV")
-	{
-		State = EWorkspaceState::Moved; // Moved/Renamed
-	}
-	else
-	{
-		UE_LOG(LogSourceControl, Warning, TEXT("Unknown file status '%s' (in line '%s')"), *FileStatus, *InResult);
-		State = EWorkspaceState::Unknown;
-	}
-
-	return State;
-}
-
-/**
- * @brief Parse the array of strings results of a 'cm status --noheaders --all --ignored' command
- *
- * Called in case of a regular status command for one or multiple files (not for a whole directory). 
- *
- * @param[in]	InFiles		List of files in a directory (never empty).
- * @param[in]	InResults	Lines of results from the "status" command
- * @param[out]	OutStates	States of files for witch the status has been gathered
- *
- * Example cm status results:
- CH Content\Changed_BP.uasset
- CO Content\CheckedOut_BP.uasset
- CP Content\Copied_BP.uasset
- RP Content\Replaced_BP.uasset
- AD Content\Added_BP.uasset
- PR Content\Private_BP.uasset
- IG Content\Ignored_BP.uasset
- DE Content\Deleted_BP.uasset
- LD Content\Deleted2_BP.uasset
- MV 100% Content\ToMove_BP.uasset -> Content\Moved_BP.uasset
- LM 100% Content\ToMove2_BP.uasset -> Content\Moved2_BP.uasset
- */
-static void ParseFileStatusResult(TArray<FString>&& InFiles, const TArray<FString>& InResults, TArray<FPlasticSourceControlState>& OutStates, int32& OutChangeset, FString& OutBranchName)
-{
-	TRACE_CPUPROFILER_EVENT_SCOPE(PlasticSourceControlUtils::ParseFileStatusResult);
-
-	const FString& WorkingDirectory = FPlasticSourceControlModule::Get().GetProvider().GetPathToWorkspaceRoot();
-
-	// Parse the first two lines with Changeset number and Branch name (the second being requested only once at init)
-	FString RepositoryName, ServerUrl;
-	ParseWorkspaceInformation(InResults, OutChangeset, RepositoryName, ServerUrl, OutBranchName);
-
-	// Iterate on each file explicitly listed in the command
-	for (FString& InFile : InFiles)
-	{
-		FPlasticSourceControlState FileState(MoveTemp(InFile));
-		const FString& File = FileState.LocalFilename;
-
-		// Search the file in the list of status
-		// NOTE: in case of rename by editor, there are two results: checked-out AND renamed
-		// => we want to get the second one, witch is always the rename, so we search from the end
-		int32 IdxResult = InResults.FindLastByPredicate(FPlasticStatusFileMatcher(File));
-		if (IdxResult != INDEX_NONE)
-		{
-			// File found in status results; only the case for "changed" files
-			FileState.WorkspaceState = StateFromPlasticStatus(InResults[IdxResult]);
-
-			// Extract the original name of a Moved/Renamed file
-			if (EWorkspaceState::Moved == FileState.WorkspaceState)
-			{
-				FileState.MovedFrom = FPaths::ConvertRelativePathToFull(WorkingDirectory, RenamedFromPlasticStatus(InResults[IdxResult]));
-			}
-		}
-		else
-		{
-			// File not found in status
-			if (FPaths::FileExists(File))
-			{
-				// usually means the file is unchanged, or is on Hidden changes
-				FileState.WorkspaceState = EWorkspaceState::Controlled; // Unchanged
-			}
-			else
-			{
-				// but also the case for newly created content: there is no file on disk until the content is saved for the first time
-				FileState.WorkspaceState = EWorkspaceState::Private; // Not Controlled
-			}
-		}
-		FileState.TimeStamp.Now();
-
-		// debug log (only for the first few files)
-		if (OutStates.Num() < 20)
-		{
-			UE_LOG(LogSourceControl, Verbose, TEXT("%s = %d:%s"), *File, static_cast<uint32>(FileState.WorkspaceState), FileState.ToString());
-		}
-
-		OutStates.Add(MoveTemp(FileState));
-	}
-	// debug log (if too many files)
-	if (OutStates.Num() > 20)
-	{
-		UE_LOG(LogSourceControl, Verbose, TEXT("[...] %d more files"), OutStates.Num() - 20);
-	}
-}
-
-/**
- * @brief Detect Deleted files in case of a "whole directory status" (no file listed in the command)
- * 
- * Parse the array of strings results of a 'cm status --noheaders --all --ignored' command
- *
- * @param[in]	InResults	Lines of results from the "status" command
- * @param[out]	OutStates	States of files for witch the status has been gathered
- *
- * @see #ParseFileStatusResult() above for an example of a cm status results
-*/
-static void ParseDirectoryStatusResultForDeleted(const TArray<FString>& InResults, TArray<FPlasticSourceControlState>& OutStates)
-{
-	TRACE_CPUPROFILER_EVENT_SCOPE(PlasticSourceControlUtils::ParseDirectoryStatusResultForDeleted);
-
-	const FString& WorkingDirectory = FPlasticSourceControlModule::Get().GetProvider().GetPathToWorkspaceRoot();
-
-	// Iterate on each line of result of the status command
-	for (const FString& Result : InResults)
-	{
-		const EWorkspaceState WorkspaceState = StateFromPlasticStatus(Result);
-		if ((EWorkspaceState::Deleted == WorkspaceState) || (EWorkspaceState::LocallyDeleted == WorkspaceState))
-		{
-			FString RelativeFilename = FilenameFromPlasticStatus(Result);
-			FString AbsoluteFilename = FPaths::ConvertRelativePathToFull(WorkingDirectory, MoveTemp(RelativeFilename));
-			FPlasticSourceControlState FileState(MoveTemp(AbsoluteFilename));
-			FileState.WorkspaceState = WorkspaceState;
-			FileState.TimeStamp.Now();
-
-			UE_LOG(LogSourceControl, Verbose, TEXT("%s = %d:%s"), *FileState.LocalFilename, static_cast<uint32>(FileState.WorkspaceState), FileState.ToString());
-
-			OutStates.Add(MoveTemp(FileState));
-		}
-	}
-}
-
-/// Visitor to list all files in subdirectory
-class FFileVisitor : public IPlatformFile::FDirectoryVisitor
-{
-public:
-	virtual bool Visit(const TCHAR* FilenameOrDirectory, bool bIsDirectory) override
-	{
-		if (!bIsDirectory)
-		{
-			Files.Add(FilenameOrDirectory);
-		}
-		return true;
-	}
-
-	TArray<FString> Files;
-};
-
-/**
- * @brief Run a "status" command for a directory to get workspace file states
+ * @brief Run a "status" command for a directory to get the local workspace file states
  *
  *  ie. Changed, CheckedOut, Copied, Replaced, Added, Private, Ignored, Deleted, LocallyDeleted, Moved, LocallyMoved
  *
  *  It is either a command for a whole directory (ie. "Content/", in case of "Submit to Source Control"),
  * or for one or more files all on a same directory (by design, since we group files by directory in RunUpdateStatus())
  *
- * @param[in]	InDir				The path to the common directory of all the files listed after.
+ * @param[in]	InDir				The path to the common directory of all the files listed after (never empty).
  * @param[in]	InFiles				List of files in a directory, or the path to the directory itself (never empty).
+ * @param[in]	InSearchType		Call "status" with "--all", or with just "--controlledchanged" when doing only a quick check following a source control operation
  * @param[out]	OutErrorMessages	Error messages from the "status" command
  * @param[out]	OutStates			States of files for witch the status has been gathered (distinct than InFiles in case of a "directory status")
  * @param[out]	OutChangeset		The current Changeset Number
  * @param[out]	OutBranchName		Name of the current checked-out branch
  */
-static bool RunStatus(const FString& InDir, TArray<FString>&& InFiles, TArray<FString>& OutErrorMessages, TArray<FPlasticSourceControlState>& OutStates, int32& OutChangeset, FString& OutBranchName)
+static bool RunStatus(const FString& InDir, TArray<FString>&& InFiles, const EStatusSearchType InSearchType, TArray<FString>& OutErrorMessages, TArray<FPlasticSourceControlState>& OutStates, int32& OutChangeset, FString& OutBranchName)
 {
 	TRACE_CPUPROFILER_EVENT_SCOPE(PlasticSourceControlUtils::RunStatus);
 
 	check(InFiles.Num() > 0);
 
 	TArray<FString> Parameters;
-	Parameters.Add(TEXT("--compact"));
-	Parameters.Add(TEXT("--noheaders"));
-	Parameters.Add(TEXT("--all"));
-	Parameters.Add(TEXT("--ignored"));
+	Parameters.Add(TEXT("--machinereadable"));
+	Parameters.Add(TEXT("--fieldseparator=\"") FILE_STATUS_SEPARATOR TEXT("\""));
+	if (InSearchType == EStatusSearchType::All)
+	{
+		// NOTE: don't use "--all" to avoid searching for --localmoved since it's the most time consuming (beside --changed)
+		// and its not used by the plugin (matching similarities doesn't seem to work with .uasset files)
+		Parameters.Add(TEXT("--controlledchanged"));
+		// TODO: add a user settings to avoid searching for --changed and --localdeleted, to work like Perforce on big projects,
+		// provided that the user understands the consequences (they won't see assets modified locally without a proper checkout)
+		Parameters.Add(TEXT("--changed"));
+		Parameters.Add(TEXT("--localdeleted"));
+		Parameters.Add(TEXT("--private"));
+		Parameters.Add(TEXT("--ignored"));
+
+		// If the version of cm is recent enough use the new --iscochanged for "CO+CH" status
+		const FPlasticSourceControlProvider& Provider = FPlasticSourceControlModule::Get().GetProvider();
+		const bool bUsesCheckedOutChanged = Provider.GetPlasticScmVersion() >= PlasticSourceControlVersions::StatusIsCheckedOutChanged;
+		if (bUsesCheckedOutChanged)
+		{
+			Parameters.Add(TEXT("--iscochanged"));
+		}
+	}
+	else if (InSearchType == EStatusSearchType::ControlledOnly)
+	{
+		Parameters.Add(TEXT("--controlledchanged"));
+	}
+
 	// "cm status" only operate on one path (file or directory) at a time, so use one common path for multiple files in a directory
 	TArray<FString> OnePath;
 	// Only one file: optim very useful for the .uproject file at the root to avoid parsing the whole repository
@@ -622,7 +277,14 @@ static bool RunStatus(const FString& InDir, TArray<FString>&& InFiles, TArray<FS
 	OutErrorMessages.Append(MoveTemp(ErrorMessages));
 	if (bResult)
 	{
-		// Normalize paths in the result (convert all '\' to '/')
+		// Parse the first line of status with the Changeset number, then remove it to work on a plain list of files
+		if (Results.Num() > 0)
+		{
+			PlasticSourceControlParsers::GetChangesetFromWorkspaceStatus(Results, OutChangeset);
+			Results.RemoveAt(0, 1, false);
+		}
+
+		// Normalize file paths in the result (convert all '\' to '/')
 		for (FString& Result : Results)
 		{
 			FPaths::NormalizeFilename(Result);
@@ -632,97 +294,50 @@ static bool RunStatus(const FString& InDir, TArray<FString>&& InFiles, TArray<FS
 		if (bWholeDirectory)
 		{
 			// 1) Special case for "status" of a directory: requires a specific parse logic.
-			//   (this is triggered by the "Submit to Source Control" top menu button)
-			// Find recursively all files in the directory: this enable getting the list of "Controlled" (unchanged) assets
-			FFileVisitor FileVisitor;
-			FPlatformFileManager::Get().GetPlatformFile().IterateDirectoryRecursively(*InDir, FileVisitor);
-			UE_LOG(LogSourceControl, Verbose, TEXT("RunStatus(%s): 1) special case for status of a directory containing %d file(s)"), *InDir, FileVisitor.Files.Num());
-			ParseFileStatusResult(MoveTemp(FileVisitor.Files), Results, OutStates, OutChangeset, OutBranchName);
-			// The above cannot detect assets removed / locally deleted since there is no file left to enumerate (either by the Content Browser or by File Manager)
-			// => so we also parse the status results to explicitly look for Removed/Deleted assets
-			if (Results.Num() > 0)
-			{
-				Results.RemoveAt(0, 1); // Before that, remove the first line (Workspace/Changeset info)
-			}
-			ParseDirectoryStatusResultForDeleted(Results, OutStates);
+			//   (this is triggered by the "Submit to Source Control" top menu button, but also for the initial check, the global Revert etc)
+			UE_LOG(LogSourceControl, Verbose, TEXT("RunStatus(%s): 1) special case for status of a directory:"), *InDir);
+			PlasticSourceControlParsers::ParseDirectoryStatusResult(InDir, Results, OutStates);
 		}
 		else
 		{
 			// 2) General case for one or more files in the same directory.
 			UE_LOG(LogSourceControl, Verbose, TEXT("RunStatus(%s...): 2) general case for %d file(s) in a directory (%s)"), *InFiles[0], InFiles.Num(), *InDir);
-			ParseFileStatusResult(MoveTemp(InFiles), Results, OutStates, OutChangeset, OutBranchName);
+			PlasticSourceControlParsers::ParseFileStatusResult(MoveTemp(InFiles), Results, OutStates);
 		}
 	}
 
 	return bResult;
 }
 
-// Parse the fileinfo output format "{RevisionChangeset};{RevisionHeadChangeset};{RepSpec};{LockedBy};{LockedWhere}"
-// for example "40;41;repo@server:port;srombauts;UEPlasticPluginDev"
-class FPlasticFileinfoParser
-{
-public:
-	explicit FPlasticFileinfoParser(const FString& InResult)
-	{
-		TArray<FString> Fileinfos;
-		const int32 NbElmts = InResult.ParseIntoArray(Fileinfos, TEXT(";"), false); // Don't cull empty values in csv
-		if (NbElmts == 5)
-		{
-			RevisionChangeset = FCString::Atoi(*Fileinfos[0]);
-			RevisionHeadChangeset = FCString::Atoi(*Fileinfos[1]);
-			RepSpec = MoveTemp(Fileinfos[2]);
-			LockedBy = UserNameToDisplayName(MoveTemp(Fileinfos[3]));
-			LockedWhere = MoveTemp(Fileinfos[4]);
-		}
-	}
-
-	int32 RevisionChangeset;
-	int32 RevisionHeadChangeset;
-	FString RepSpec;
-	FString LockedBy;
-	FString LockedWhere;
-};
-
-/** Parse the array of strings result of a 'cm fileinfo --format="{RevisionChangeset};{RevisionHeadChangeset};{RepSpec};{LockedBy};{LockedWhere}"' command
- *
- * Example cm fileinfo results:
-16;16;;
-14;15;;
-17;17;srombauts;Workspace_2
- */
-static void ParseFileinfoResults(const TArray<FString>& InResults, TArray<FPlasticSourceControlState>& InOutStates)
+bool RunListSmartLocks(const FString& InRepository, TMap<FString, PlasticSourceControlParsers::FSmartLockInfoParser>& OutSmartLocks)
 {
-	TRACE_CPUPROFILER_EVENT_SCOPE(PlasticSourceControlUtils::ParseFileinfoResults);
-
-	const FPlasticSourceControlProvider& Provider = FPlasticSourceControlModule::Get().GetProvider();
+	TRACE_CPUPROFILER_EVENT_SCOPE(PlasticSourceControlUtils::RunListSmartLocks);
 
-	ensureMsgf(InResults.Num() == InOutStates.Num(), TEXT("The fileinfo command should gives the same number of infos as the status command"));
+	TArray<FString> Results;
+	TArray<FString> ErrorMessages;
+	TArray<FString> Parameters;
+	Parameters.Add(TEXT("list"));
+	Parameters.Add(TEXT("--machinereadable"));
+	Parameters.Add(TEXT("--smartlocks"));
+	Parameters.Add(TEXT("--anystatus"));
+	Parameters.Add(TEXT("--fieldseparator=\"") FILE_STATUS_SEPARATOR TEXT("\""));
+	Parameters.Add(TEXT("--dateformat=yyyy-MM-ddTHH:mm:ss"));
+	bool bResult = RunCommand(TEXT("lock"), Parameters, TArray<FString>(), Results, ErrorMessages);
 
-	// Iterate on all files and all status of the result (assuming same number of line of results than number of file states)
-	for (int32 IdxResult = 0; IdxResult < InResults.Num(); IdxResult++)
+	if (bResult)
 	{
-		const FString& Fileinfo = InResults[IdxResult];
-		FPlasticSourceControlState& FileState = InOutStates[IdxResult];
-		const FString& File = FileState.LocalFilename;
-		FPlasticFileinfoParser FileinfoParser(Fileinfo);
-
-		FileState.LocalRevisionChangeset = FileinfoParser.RevisionChangeset;
-		FileState.DepotRevisionChangeset = FileinfoParser.RevisionHeadChangeset;
-		FileState.RepSpec = FileinfoParser.RepSpec;
-		FileState.LockedBy = MoveTemp(FileinfoParser.LockedBy);
-		FileState.LockedWhere = MoveTemp(FileinfoParser.LockedWhere);
-
-		// debug log (only for the first few files)
-		if (IdxResult < 20)
+		for (int32 IdxResult = 0; IdxResult < Results.Num(); IdxResult++)
 		{
-			UE_LOG(LogSourceControl, Verbose, TEXT("%s: %d;%d %s by '%s' (%s)"), *File, FileState.LocalRevisionChangeset, FileState.DepotRevisionChangeset, *FileState.RepSpec, *FileState.LockedBy, *FileState.LockedWhere);
+			const FString& Result = Results[IdxResult];
+			PlasticSourceControlParsers::FSmartLockInfoParser SmartLockInfoParser(Result);
+			if (SmartLockInfoParser.Repository == InRepository)
+			{
+				OutSmartLocks.Add(SmartLockInfoParser.Filename, SmartLockInfoParser);
+			}
 		}
 	}
-	// debug log (if too many files)
-	if (InResults.Num() > 20)
-	{
-		UE_LOG(LogSourceControl, Verbose, TEXT("[...] %d more files"), InResults.Num() - 20);
-	}
+
+	return bResult;
 }
 
 /**
@@ -779,12 +394,12 @@ static bool RunFileinfo(const bool bInWholeDirectory, const bool bInUpdateHistor
 		TArray<FString> Results;
 		TArray<FString> ErrorMessages;
 		TArray<FString> Parameters;
-		Parameters.Add(TEXT("--format=\"{RevisionChangeset};{RevisionHeadChangeset};{RepSpec};{LockedBy};{LockedWhere}\""));
+		Parameters.Add(TEXT("--format=\"{RevisionChangeset};{RevisionHeadChangeset};{RepSpec};{LockedBy};{LockedWhere};{ServerPath}\""));
 		bResult = RunCommand(TEXT("fileinfo"), Parameters, SelectedFiles, Results, ErrorMessages);
 		OutErrorMessages.Append(MoveTemp(ErrorMessages));
 		if (bResult)
 		{
-			ParseFileinfoResults(Results, SelectedStates);
+			PlasticSourceControlParsers::ParseFileinfoResults(Results, SelectedStates);
 			InOutStates.Append(MoveTemp(SelectedStates));
 		}
 	}
@@ -792,40 +407,6 @@ static bool RunFileinfo(const bool bInWholeDirectory, const bool bInUpdateHistor
 	return bResult;
 }
 
-// FILE_CONFLICT /Content/FirstPersonBP/Blueprints/FirstPersonProjectile.uasset 1 4 6 903
-// (explanations: 'The file /Content/FirstPersonBP/Blueprints/FirstPersonProjectile.uasset needs to be merged from cs:4 to cs:6 base cs:1. Changed by both contributors.')
-class FPlasticMergeConflictParser
-{
-public:
-	explicit FPlasticMergeConflictParser(const FString& InResult)
-	{
-		static const FString FILE_CONFLICT(TEXT("FILE_CONFLICT "));
-		if (InResult.StartsWith(FILE_CONFLICT, ESearchCase::CaseSensitive))
-		{
-			FString Temp = InResult.RightChop(FILE_CONFLICT.Len());
-			int32 WhitespaceIndex;
-			if (Temp.FindChar(TEXT(' '), WhitespaceIndex))
-			{
-				Filename = Temp.Left(WhitespaceIndex);
-			}
-			Temp.RightChopInline(WhitespaceIndex + 1);
-			if (Temp.FindChar(TEXT(' '), WhitespaceIndex))
-			{
-				BaseChangeset = Temp.Left(WhitespaceIndex);
-			}
-			Temp.RightChopInline(WhitespaceIndex + 1);
-			if (Temp.FindChar(TEXT(' '), WhitespaceIndex))
-			{
-				SourceChangeset = Temp.Left(WhitespaceIndex);
-			}
-		}
-	}
-
-	FString Filename;
-	FString BaseChangeset;
-	FString SourceChangeset;
-};
-
 // Check if merging, and from which changelist, then execute a cm merge command to amend status for listed files
 static bool RunCheckMergeStatus(const TArray<FString>& InFiles, TArray<FString>& OutErrorMessages, TArray<FPlasticSourceControlState>& OutStates)
 {
@@ -891,7 +472,7 @@ static bool RunCheckMergeStatus(const TArray<FString>& InFiles, TArray<FString>&
 					// Parse the result, one line for each conflicted files:
 					for (const FString& Result : Results)
 					{
-						FPlasticMergeConflictParser MergeConflict(Result);
+						PlasticSourceControlParsers::FPlasticMergeConflictParser MergeConflict(Result);
 						UE_LOG(LogSourceControl, Log, TEXT("MergeConflict.Filename: '%s'"), *MergeConflict.Filename);
 						for (FPlasticSourceControlState& State : OutStates)
 						{
@@ -946,7 +527,7 @@ struct FFilesInCommonDir
 };
 
 // Run a batch of Plastic "status" and "fileinfo" commands to update status of given files and directories.
-bool RunUpdateStatus(const TArray<FString>& InFiles, const bool bInUpdateHistory, TArray<FString>& OutErrorMessages, TArray<FPlasticSourceControlState>& OutStates, int32& OutChangeset, FString& OutBranchName)
+bool RunUpdateStatus(const TArray<FString>& InFiles, const EStatusSearchType InSearchType, const bool bInUpdateHistory, TArray<FString>& OutErrorMessages, TArray<FPlasticSourceControlState>& OutStates, int32& OutChangeset, FString& OutBranchName)
 {
 	bool bResults = true;
 
@@ -1041,7 +622,7 @@ bool RunUpdateStatus(const TArray<FString>& InFiles, const bool bInUpdateHistory
 		// Run a "status" command on the directory to get workspace file states.
 		// (ie. Changed, CheckedOut, Copied, Replaced, Added, Private, Ignored, Deleted, LocallyDeleted, Moved, LocallyMoved)
 		TArray<FPlasticSourceControlState> States;
-		const bool bGroupOk = RunStatus(Group.Value.CommonDir, MoveTemp(Group.Value.Files), OutErrorMessages, States, OutChangeset, OutBranchName);
+		const bool bGroupOk = RunStatus(Group.Value.CommonDir, MoveTemp(Group.Value.Files), InSearchType, OutErrorMessages, States, OutChangeset, OutBranchName);
 		if (!bGroupOk)
 		{
 			bResults = false;
@@ -1049,7 +630,7 @@ bool RunUpdateStatus(const TArray<FString>& InFiles, const bool bInUpdateHistory
 		else if (States.Num() > 0)
 		{
 			// Run a "fileinfo" command to update complementary status information of given files.
-			// (ie RevisionChangeset, RevisionHeadChangeset, RepSpec, LockedBy, LockedWhere)
+			// (ie RevisionChangeset, RevisionHeadChangeset, RepSpec, LockedBy, LockedWhere, ServerPath)
 			// In case of "whole directory status", there is no explicit file in the group (it contains only the directory)
 			// => work on the list of files discovered by RunStatus()
 			bResults &= RunFileinfo(bWholeDirectory, bInUpdateHistory, OutErrorMessages, States);
@@ -1063,290 +644,61 @@ bool RunUpdateStatus(const TArray<FString>& InFiles, const bool bInUpdateHistory
 	return bResults;
 }
 
-// Run a Plastic "cat" command to dump the binary content of a revision into a file.
-// cm cat revid:1230@rep:myrep@repserver:myserver:8084 --raw --file=Name124.tmp
-bool RunDumpToFile(const FString& InPathToPlasticBinary, const FString& InRevSpec, const FString& InDumpFileName)
+// Run a "getfile" command to dump the binary content of a revision into a file.
+bool RunGetFile(const FString& InRevSpec, const FString& InDumpFileName)
 {
-	TRACE_CPUPROFILER_EVENT_SCOPE(PlasticSourceControlUtils::RunDumpToFile);
+	TRACE_CPUPROFILER_EVENT_SCOPE(PlasticSourceControlUtils::RunGetFile);
 
 	int32	ReturnCode = 0;
 	FString Results;
 	FString Errors;
 
-	// start with the Plastic command itself, then add revspec and temp filename to dump
-	FString FullCommand = TEXT("cat \"");
-	FullCommand += InRevSpec;
-	FullCommand += TEXT("\" --raw --file=\"");
-	FullCommand += InDumpFileName;
-	FullCommand += TEXT("\"");
-
-	UE_LOG(LogSourceControl, Verbose, TEXT("RunDumpToFile: '%s %s'"), *InPathToPlasticBinary, *FullCommand);
-	const bool bResult = FPlatformProcess::ExecProcess(*InPathToPlasticBinary, *FullCommand, &ReturnCode, &Results, &Errors);
-	UE_LOG(LogSourceControl, Log, TEXT("RunDumpToFile: ExecProcess ReturnCode=%d Results='%s'"), ReturnCode, *Results);
-	if (!bResult || !Errors.IsEmpty())
-	{
-		UE_LOG(LogSourceControl, Error, TEXT("RunDumpToFile: ExecProcess ReturnCode=%d Errors='%s'"), ReturnCode, *Errors);
-	}
+	TArray<FString> Parameters;
+	Parameters.Add(FString::Printf(TEXT("\"%s\""), *InRevSpec));
+	Parameters.Add(TEXT("--raw"));
+	Parameters.Add(FString::Printf(TEXT("--file=\"%s\""), *InDumpFileName));
+	const bool bResult = PlasticSourceControlUtils::RunCommand(TEXT("getfile"), Parameters, TArray<FString>(), Results, Errors);
 
 	return bResult;
 }
 
-FString DecodeXmlEntities(const FString& InString)
+// Convert a file state to a string ala Perforce, see also ParseShelveFileStatus()
+FString FileStateToAction(const EWorkspaceState InState)
 {
-	FString String = InString;
-	String.ReplaceInline(TEXT("&amp;"), TEXT("&"));
-	String.ReplaceInline(TEXT("&quot;"), TEXT("\""));
-	String.ReplaceInline(TEXT("&apos;"), TEXT("'"));
-	String.ReplaceInline(TEXT("&lt;"), TEXT("<"));
-	String.ReplaceInline(TEXT("&gt;"), TEXT(">"));
-	return String;
+	switch (InState)
+	{
+	case EWorkspaceState::Added:
+		return TEXT("add");
+	case EWorkspaceState::Deleted:
+		return TEXT("delete");
+	case EWorkspaceState::Moved:
+		return TEXT("branch");
+	case EWorkspaceState::CheckedOutChanged:
+	default:
+		return TEXT("edit");
+	}
 }
 
-/**
- * Parse results of the 'cm history --moveddeleted --xml --encoding="utf-8"' command.
- * 
- * Results of the history command looks like that:
-<RevisionHistoriesResult>
-  <RevisionHistories>
-	<RevisionHistory>
-	  <ItemName>C:/Workspace/UE4PlasticPluginDev/Content/FirstPersonBP/Blueprints/BP_TestsRenamed.uasset</ItemName>
-	  <Revisions>
-		<Revision>
-		  <RevisionSpec>C:/Workspace/UE4PlasticPluginDev/Content/FirstPersonBP/Blueprints/BP_TestsRenamed.uasset#cs:7</RevisionSpec>
-		  <Branch>/main</Branch>
-		  <CreationDate>2019-10-14T09:52:07+02:00</CreationDate>
-		  <RevisionType>bin</RevisionType>
-		  <ChangesetNumber>7</ChangesetNumber>
-		  <Owner>SRombauts</Owner>
-		  <Comment>New tests</Comment>
-		  <Repository>UE4PlasticPluginDev</Repository>
-		  <Server>localhost:8087</Server>
-		  <RepositorySpec>UE4PlasticPluginDev@localhost:8087</RepositorySpec>
-		  <Size>22356</Size>
-		  <Hash>zzuB6G9fbWz1md12+tvBxg==</Hash>
-		</Revision>
-		...
-		<Revision>
-		  <RevisionSpec>C:/Workspace/UE4PlasticPluginDev/Content/FirstPersonBP/Blueprints/BP_TestsRenamed.uasset#cs:12</RevisionSpec>
-		  <Branch>Removed /Content/FirstPersonBP/Blueprints/BP_TestsRenamed.uasset</Branch>
-		  <CreationDate>2022-04-28T16:00:37+02:00</CreationDate>
-		  <RevisionType />
-		  <ChangesetNumber>12</ChangesetNumber>
-		  <Owner>sebastien.rombauts</Owner>
-		  <Comment />
-		  <Repository>UE4PlasticPluginDev</Repository>
-		  <Server>localhost:8087</Server>
-		  <RepositorySpec>UE4PlasticPluginDev@localhost:8087</RepositorySpec>
-		  <Size>22406</Size>
-		  <Hash>uR7NdDRAyKqADdyAqh67Rg==</Hash>
-		</Revision>
-
-	  </Revisions>
-	</RevisionHistory>
-	<RevisionHistory>
-	  <ItemName>C:/Workspace/UE4PlasticPluginDev/Content/FirstPersonBP/Blueprints/BP_YetAnother.uasset</ItemName>
-		...
-	</RevisionHistory>
-  </RevisionHistories>
-</RevisionHistoriesResult>
-*/
-static bool ParseHistoryResults(const bool bInUpdateHistory, const FXmlFile& InXmlResult, TArray<FPlasticSourceControlState>& InOutStates)
+// Run a Plastic "history" command and parse it's XML result.
+bool RunGetHistory(const bool bInUpdateHistory, TArray<FPlasticSourceControlState>& InOutStates, TArray<FString>& OutErrorMessages)
 {
-	TRACE_CPUPROFILER_EVENT_SCOPE(PlasticSourceControlUtils::ParseHistoryResults);
-
-	const FPlasticSourceControlProvider& Provider = FPlasticSourceControlModule::Get().GetProvider();
-	const FString RootRepSpec = FString::Printf(TEXT("%s@%s"), *Provider.GetRepositoryName(), *Provider.GetServerUrl());
-
-	static const FString RevisionHistoriesResult(TEXT("RevisionHistoriesResult"));
-	static const FString RevisionHistories(TEXT("RevisionHistories"));
-	static const FString RevisionHistory(TEXT("RevisionHistory"));
-	static const FString ItemName(TEXT("ItemName"));
-	static const FString Revisions(TEXT("Revisions"));
-	static const FString Revision(TEXT("Revision"));
-	static const FString Branch(TEXT("Branch"));
-	static const FString CreationDate(TEXT("CreationDate"));
-	static const FString RevisionType(TEXT("RevisionType"));
-	static const FString ChangesetNumber(TEXT("ChangesetNumber"));
-	static const FString Owner(TEXT("Owner"));
-	static const FString Comment(TEXT("Comment"));
-	static const FString Size(TEXT("Size"));
-	static const FString Hash(TEXT("Hash"));
-
-	const FXmlNode* RevisionHistoriesResultNode = InXmlResult.GetRootNode();
-	if (RevisionHistoriesResultNode == nullptr || RevisionHistoriesResultNode->GetTag() != RevisionHistoriesResult)
-	{
-		return false;
-	}
+	TRACE_CPUPROFILER_EVENT_SCOPE(PlasticSourceControlUtils::RunGetHistory);
 
-	const FXmlNode* RevisionHistoriesNode = RevisionHistoriesResultNode->FindChildNode(RevisionHistories);
-	if (RevisionHistoriesNode == nullptr)
+	bool bResult = true;
+	FString Results;
+	FString Errors;
+	TArray<FString> Parameters;
+	// Detecting move and deletion is costly as it is implemented as two extra queries to the server; do it only when getting the history of the current branch
+	if (bInUpdateHistory)
 	{
-		return false;
+		Parameters.Add(TEXT("--moveddeleted"));
 	}
-
-	const TArray<FXmlNode*>& RevisionHistoryNodes = RevisionHistoriesNode->GetChildrenNodes();
-	for (const FXmlNode* RevisionHistoryNode : RevisionHistoryNodes)
+	Parameters.Add(TEXT("--xml"));
+	Parameters.Add(TEXT("--encoding=\"utf-8\""));
+	const FPlasticSourceControlProvider& Provider = FPlasticSourceControlModule::Get().GetProvider();
+	if (Provider.GetPlasticScmVersion() >= PlasticSourceControlVersions::NewHistoryLimit)
 	{
-		const FXmlNode* ItemNameNode = RevisionHistoryNode->FindChildNode(ItemName);
-		if (ItemNameNode == nullptr)
-		{
-			continue;
-		}
-
-		const FString Filename = ItemNameNode->GetContent();
-		FPlasticSourceControlState* InOutStatePtr = InOutStates.FindByPredicate(
-			[&Filename](const FPlasticSourceControlState& State) { return State.LocalFilename == Filename; }
-		);
-		if (InOutStatePtr == nullptr)
-		{
-			continue;
-		}
-		FPlasticSourceControlState& InOutState = *InOutStatePtr;
-
-		const FXmlNode* RevisionsNode = RevisionHistoryNode->FindChildNode(Revisions);
-		if (RevisionsNode == nullptr)
-		{
-			continue;
-		}
-
-		const TArray<FXmlNode*>& RevisionNodes = RevisionsNode->GetChildrenNodes();
-		if (bInUpdateHistory)
-		{
-			InOutState.History.Reserve(RevisionNodes.Num());
-		}
-
-		// parse history in reverse: needed to get most recent at the top (implied by the UI)
-		// Note: limit to last 100 changes, like Perforce
-		static const int32 MaxRevisions = 100;
-		const int32 MinIndex = FMath::Max(0, RevisionNodes.Num() - MaxRevisions);
-		for (int32 Index = RevisionNodes.Num() - 1; Index >= MinIndex; Index--)
-		{
-			if (const FXmlNode* RevisionNode = RevisionNodes[Index])
-			{
-				const TSharedRef<FPlasticSourceControlRevision, ESPMode::ThreadSafe> SourceControlRevision = MakeShared<FPlasticSourceControlRevision>();
-				SourceControlRevision->State = &InOutState;
-				SourceControlRevision->Filename = Filename;
-				SourceControlRevision->RevisionId = Index + 1;
-
-				if (const FXmlNode* RevisionTypeNode = RevisionNode->FindChildNode(RevisionType))
-				{
-					if (!RevisionTypeNode->GetContent().IsEmpty())
-					{
-						if (Index == 0)
-						{
-							SourceControlRevision->Action = TEXT("add");
-						}
-						else
-						{
-							SourceControlRevision->Action = TEXT("edit");
-						}
-					}
-					else
-					{
-						SourceControlRevision->Action = TEXT("delete");
-					}
-				}
-
-				if (const FXmlNode* ChangesetNumberNode = RevisionNode->FindChildNode(ChangesetNumber))
-				{
-					const FString& Changeset = ChangesetNumberNode->GetContent();
-					SourceControlRevision->ChangesetNumber = FCString::Atoi(*Changeset); // Value now used in the Revision column and in the Asset Menu History
-
-					// Also append depot name to the revision, but only when it is different from the default one (ie for xlinks sub repository)
-					if (!InOutState.RepSpec.IsEmpty() && (InOutState.RepSpec != RootRepSpec))
-					{
-						TArray<FString> RepSpecs;
-						InOutState.RepSpec.ParseIntoArray(RepSpecs, TEXT("@"));
-						SourceControlRevision->Revision = FString::Printf(TEXT("cs:%s@%s"), *Changeset, *RepSpecs[0]);
-					}
-					else
-					{
-						SourceControlRevision->Revision = FString::Printf(TEXT("cs:%s"), *Changeset);
-					}
-				}
-				if (const FXmlNode* CommentNode = RevisionNode->FindChildNode(Comment))
-				{
-					SourceControlRevision->Description = DecodeXmlEntities(CommentNode->GetContent());
-				}
-				if (const FXmlNode* OwnerNode = RevisionNode->FindChildNode(Owner))
-				{
-					SourceControlRevision->UserName = UserNameToDisplayName(OwnerNode->GetContent());
-				}
-				if (const FXmlNode* DateNode = RevisionNode->FindChildNode(CreationDate))
-				{
-					FString DateIso = DateNode->GetContent();
-					const int len = DateIso.Len();
-					if (DateIso.Len() > 29)
-					{	//                           |--|
-						//    2016-04-18T10:44:49.0000000+02:00
-						// => 2016-04-18T10:44:49.000+02:00
-						DateIso = DateNode->GetContent().LeftChop(10) + DateNode->GetContent().RightChop(27);
-					}
-					FDateTime::ParseIso8601(*DateIso, SourceControlRevision->Date);
-				}
-				if (const FXmlNode* BranchNode = RevisionNode->FindChildNode(Branch))
-				{
-					SourceControlRevision->Branch = BranchNode->GetContent();
-				}
-				if (const FXmlNode* SizeNode = RevisionNode->FindChildNode(Size))
-				{
-					SourceControlRevision->FileSize = FCString::Atoi(*SizeNode->GetContent());
-				}
-
-				// Detect and skip more recent changesets on other branches (ie above the RevisionHeadChangeset)
-				// since we usually don't want to display changes from other branches in the History window...
-				// except in case of a merge conflict, where the Editor expects the tip of the "source (remote)" branch to be at the top of the history!
-				if (   (SourceControlRevision->ChangesetNumber > InOutState.DepotRevisionChangeset)
-					&& (SourceControlRevision->GetRevision() != InOutState.PendingResolveInfo.RemoteRevision))
-				{
-					InOutState.HeadBranch = SourceControlRevision->Branch;
-					InOutState.HeadAction = SourceControlRevision->Action;
-					InOutState.HeadChangeList = SourceControlRevision->ChangesetNumber;
-					InOutState.HeadUserName = SourceControlRevision->UserName;
-					InOutState.HeadModTime = SourceControlRevision->Date.ToUnixTimestamp();
-				}
-				else if (bInUpdateHistory)
-				{
-					InOutState.History.Add(SourceControlRevision);
-				}
-
-				// Also grab the UserName of the author of the current depot/head changeset
-				if ((SourceControlRevision->ChangesetNumber == InOutState.DepotRevisionChangeset) && InOutState.HeadUserName.IsEmpty())
-				{
-					InOutState.HeadUserName = SourceControlRevision->UserName;
-				}
-
-				if (!bInUpdateHistory)
-				{
-					break; // if not updating the history, just getting the head of the latest branch is enough
-				}
-			}
-		}
-	}
-
-	return true;
-}
-
-// Run a Plastic "history" command and parse it's XML result.
-bool RunGetHistory(const bool bInUpdateHistory, TArray<FPlasticSourceControlState>& InOutStates, TArray<FString>& OutErrorMessages)
-{
-	TRACE_CPUPROFILER_EVENT_SCOPE(PlasticSourceControlUtils::RunGetHistory);
-
-	bool bResult = true;
-	FString Results;
-	FString Errors;
-	TArray<FString> Parameters;
-	// Detecting move and deletion is costly as it is implemented as two extra queries to the server; do it only when getting the history of the current branch
-	if (bInUpdateHistory)
-	{
-		Parameters.Add(TEXT("--moveddeleted"));
-	}
-	Parameters.Add(TEXT("--xml"));
-	Parameters.Add(TEXT("--encoding=\"utf-8\""));
-	const FPlasticSourceControlProvider& Provider = FPlasticSourceControlModule::Get().GetProvider();
-	if (Provider.GetPlasticScmVersion() >= PlasticSourceControlVersions::NewHistoryLimit)
-	{
-		if (bInUpdateHistory)
+		if (bInUpdateHistory)
 		{
 			// --limit=0 will not limit the number of revisions, as stated by LimitNumberOfRevisionsInHistory
 			Parameters.Add(FString::Printf(TEXT("--limit=%d"), GetDefault<UPlasticSourceControlProjectSettings>()->LimitNumberOfRevisionsInHistory));
@@ -1376,19 +728,7 @@ bool RunGetHistory(const bool bInUpdateHistory, TArray<FPlasticSourceControlStat
 		bResult = RunCommand(TEXT("history"), Parameters, Files, Results, Errors);
 		if (bResult)
 		{
-			FXmlFile XmlFile;
-			{
-				TRACE_CPUPROFILER_EVENT_SCOPE(PlasticSourceControlUtils::RunGetHistory::FXmlFile::LoadFile);
-				bResult = XmlFile.LoadFile(Results, EConstructMethod::ConstructFromBuffer);
-			}
-			if (bResult)
-			{
-				bResult = ParseHistoryResults(bInUpdateHistory, XmlFile, InOutStates);
-			}
-			else
-			{
-				UE_LOG(LogSourceControl, Error, TEXT("RunGetHistory: XML parse error '%s'"), *XmlFile.GetLastError())
-			}
+			bResult = PlasticSourceControlParsers::ParseHistoryResults(bInUpdateHistory, Results, InOutStates);
 		}
 		if (!Errors.IsEmpty())
 		{
@@ -1399,86 +739,6 @@ bool RunGetHistory(const bool bInUpdateHistory, TArray<FPlasticSourceControlStat
 	return bResult;
 }
 
-
-/* Parse results of the 'cm update --xml=tempfile.xml --encoding="utf-8"' command.
- * 
- * Results of the update command looks like that:
-<UpdatedItems>
-  <List>
-	<UpdatedItem>
-	  <Path>c:\Workspace\UE5PlasticPluginDev\Content\NewFolder\BP_CheckedOut.uasset</Path>
-	  <User>[REDACTED_EMAIL]</User>
-	  <Changeset>94</Changeset>
-	  <Date>2022-10-27T11:58:02+02:00</Date>
-	</UpdatedItem>
-  </List>
-</UpdatedItems>
-*/
-static bool ParseUpdateResults(const FXmlFile& InXmlResult, TArray<FString>& OutFiles)
-{
-	TRACE_CPUPROFILER_EVENT_SCOPE(PlasticSourceControlUtils::ParseUpdateResults);
-
-	static const FString UpdatedItems(TEXT("UpdatedItems"));
-	static const FString List(TEXT("List"));
-	static const FString UpdatedItem(TEXT("UpdatedItem"));
-	static const FString Path(TEXT("Path"));
-
-	const FXmlNode* UpdatedItemsNode = InXmlResult.GetRootNode();
-	if (UpdatedItemsNode == nullptr || UpdatedItemsNode->GetTag() != UpdatedItems)
-	{
-		return false;
-	}
-
-	const FXmlNode* ListNode = UpdatedItemsNode->FindChildNode(List);
-	if (ListNode == nullptr)
-	{
-		return false;
-	}
-
-	const TArray<FXmlNode*>& UpdatedItemNodes = ListNode->GetChildrenNodes();
-	for (const FXmlNode* UpdatedItemNode : UpdatedItemNodes)
-	{
-		if (const FXmlNode* PathNode = UpdatedItemNode->FindChildNode(Path))
-		{
-			FString Filename = PathNode->GetContent();
-			FPaths::NormalizeFilename(Filename);
-			OutFiles.Add(Filename);
-		}
-	}
-
-	return true;
-}
-
-
-/* Parse results of the 'cm partial update --report --machinereadable' command.
- *
- * Results of the update command looks like that:
-STAGE Plastic is updating your workspace. Wait a moment, please...
-STAGE Updated 63.01 KB of 63.01 KB (12 of 12 files to download / 16 of 21 operations to apply) /Content/Collections/SebSharedCollection.collection
-AD c:\Workspace\UE5PlasticPluginDev\Content\LevelPrototyping\Materials\MI_Solid_Red.uasset
-CH c:\Workspace\UE5PlasticPluginDev\Config\DefaultEditor.ini
-DE c:\Workspace\UE5PlasticPluginDev\Content\Collections\SebSharedCollection.collection
-*/
-static bool ParseUpdateResults(const TArray<FString>& InResults, TArray<FString>& OutFiles)
-{
-	TRACE_CPUPROFILER_EVENT_SCOPE(PlasticSourceControlUtils::ParseUpdateResultsString);
-
-	static const FString Stage = TEXT("STAGE ");
-	static const int32 PrefixLen = 3; // "XX " typically "CH ", "AD " or "DE "
-
-	for (const FString& Result : InResults)
-	{
-		if (Result.StartsWith(Stage))
-			continue;
-
-		FString Filename = Result.RightChop(PrefixLen);
-		FPaths::NormalizeFilename(Filename);
-		OutFiles.Add(Filename);
-	}
-
-	return true;
-}
-
 // Run a Plastic "update" command to sync the workspace and parse its XML results.
 bool RunUpdate(const TArray<FString>& InFiles, const bool bInIsPartialWorkspace, TArray<FString>& OutUpdatedFiles, TArray<FString>& OutErrorMessages)
 {
@@ -1489,9 +749,9 @@ bool RunUpdate(const TArray<FString>& InFiles, const bool bInIsPartialWorkspace,
 	// Detect special case for a partial checkout (CS:-1 in Gluon mode)!
 	if (!bInIsPartialWorkspace)
 	{
-		const FScopedTempFile TempFile;
+		const FScopedTempFile UpdateResultFile;
 		TArray<FString> InfoMessages;
-		Parameters.Add(FString::Printf(TEXT("--xml=\"%s\""), *TempFile.GetFilename()));
+		Parameters.Add(FString::Printf(TEXT("--xml=\"%s\""), *UpdateResultFile.GetFilename()));
 		Parameters.Add(TEXT("--encoding=\"utf-8\""));
 		Parameters.Add(TEXT("--last"));
 		Parameters.Add(TEXT("--dontmerge"));
@@ -1500,21 +760,9 @@ bool RunUpdate(const TArray<FString>& InFiles, const bool bInIsPartialWorkspace,
 		{
 			// Load and parse the result of the update command
 			FString Results;
-			if (FFileHelper::LoadFileToString(Results, *TempFile.GetFilename()))
+			if (FFileHelper::LoadFileToString(Results, *UpdateResultFile.GetFilename()))
 			{
-				FXmlFile XmlFile;
-				{
-					TRACE_CPUPROFILER_EVENT_SCOPE(PlasticSourceControlUtils::RunUpdate::FXmlFile::LoadFile);
-					bResult = XmlFile.LoadFile(Results, EConstructMethod::ConstructFromBuffer);
-				}
-				if (bResult)
-				{
-					bResult = ParseUpdateResults(XmlFile, OutUpdatedFiles);
-				}
-				else
-				{
-					UE_LOG(LogSourceControl, Error, TEXT("RunUpdate: XML parse error '%s'"), *XmlFile.GetLastError())
-				}
+				bResult = PlasticSourceControlParsers::ParseUpdateResults(Results, OutUpdatedFiles);
 			}
 		}
 	}
@@ -1526,143 +774,13 @@ bool RunUpdate(const TArray<FString>& InFiles, const bool bInIsPartialWorkspace,
 		bResult = PlasticSourceControlUtils::RunCommand(TEXT("partial update"), Parameters, InFiles, Results, OutErrorMessages);
 		if (bResult)
 		{
-			bResult = ParseUpdateResults(Results, OutUpdatedFiles);
+			bResult = PlasticSourceControlParsers::ParseUpdateResults(Results, OutUpdatedFiles);
 		}
 	}
 
 	return bResult;
 }
 
-/**
- * Parse results of the 'cm status --changelists --controlledchanged --xml --encoding="utf-8"' command.
- *
- * Results of the status changelists command looks like that:
-<StatusOutput>
-  <WorkspaceStatus>
-	<Status>
-	  <RepSpec>
-		<Server>test@cloud</Server>
-		<Name>UEPlasticPluginDev</Name>
-	  </RepSpec>
-	  <Changeset>11</Changeset>
-	</Status>
-  </WorkspaceStatus>
-  <WkConfigType>Branch</WkConfigType>
-  <WkConfigName>/main@rep:UEPlasticPluginDev@repserver:test@cloud</WkConfigName>
-  <Changelists>
-	<Changelist>
-	  <Name>Default</Name>
-	  <Description>Default Plastic SCM changelist</Description>
-	  <Changes>
-		<Change>
-		  <Type>CO</Type>
-		  <TypeVerbose>Checked-out</TypeVerbose>
-		  <Path>UEPlasticPluginDev.uproject</Path>
-		  <OldPath />
-		  <PrintableMovedPath />
-		  <MergesInfo />
-		  <SimilarityPerUnit>0</SimilarityPerUnit>
-		  <Similarity />
-		  <Size>583</Size>
-		  <PrintableSize>583 bytes</PrintableSize>
-		  <PrintableLastModified>6 days ago</PrintableLastModified>
-		  <RevisionType>enTextFile</RevisionType>
-		  <LastModified>2022-06-07T12:28:32+02:00</LastModified>
-		</Change>
-		[...]
-		<Change>
-		</Change>
-	  </Changes>
-	</Changelist>
-  </Changelists>
-</StatusOutput>
-*/
-static bool ParseChangelistsResults(const FXmlFile& InXmlResult, TArray<FPlasticSourceControlChangelistState>& OutChangelistsStates, TArray<TArray<FPlasticSourceControlState>>& OutCLFilesStates)
-{
-	TRACE_CPUPROFILER_EVENT_SCOPE(PlasticSourceControlUtils::ParseChangelistsResults);
-
-	static const FString StatusOutput(TEXT("StatusOutput"));
-	static const FString WkConfigType(TEXT("WkConfigType"));
-	static const FString WkConfigName(TEXT("WkConfigName"));
-	static const FString Changelists(TEXT("Changelists"));
-	static const FString Changelist(TEXT("Changelist"));
-	static const FString Name(TEXT("Name"));
-	static const FString Description(TEXT("Description"));
-	static const FString Changes(TEXT("Changes"));
-	static const FString Change(TEXT("Change"));
-	static const FString Type(TEXT("Type"));
-	static const FString Path(TEXT("Path"));
-
-	const FString& WorkingDirectory = FPlasticSourceControlModule::Get().GetProvider().GetPathToWorkspaceRoot();
-
-	const FXmlNode* StatusOutputNode = InXmlResult.GetRootNode();
-	if (StatusOutputNode == nullptr || StatusOutputNode->GetTag() != StatusOutput)
-	{
-		return false;
-	}
-
-	const FXmlNode* ChangelistsNode = StatusOutputNode->FindChildNode(Changelists);
-	if (ChangelistsNode)
-	{
-		const TArray<FXmlNode*>& ChangelistNodes = ChangelistsNode->GetChildrenNodes();
-		OutCLFilesStates.SetNum(ChangelistNodes.Num());
-		for (int32 ChangelistIndex = 0; ChangelistIndex < ChangelistNodes.Num(); ChangelistIndex++)
-		{
-			const FXmlNode* ChangelistNode = ChangelistNodes[ChangelistIndex];
-			check(ChangelistNode);
-			const FXmlNode* NameNode = ChangelistNode->FindChildNode(Name);
-			const FXmlNode* DescriptionNode = ChangelistNode->FindChildNode(Description);
-			const FXmlNode* ChangesNode = ChangelistNode->FindChildNode(Changes);
-			if (NameNode == nullptr || DescriptionNode == nullptr || ChangesNode == nullptr)
-			{
-				continue;
-			}
-
-			FString NameTemp = PlasticSourceControlUtils::DecodeXmlEntities(NameNode->GetContent());
-			FPlasticSourceControlChangelist ChangelistTemp(MoveTemp(NameTemp), true);
-			FString DescriptionTemp = ChangelistTemp.IsDefault() ? FString() : PlasticSourceControlUtils::DecodeXmlEntities(DescriptionNode->GetContent());
-			FPlasticSourceControlChangelistState ChangelistState(MoveTemp(ChangelistTemp), MoveTemp(DescriptionTemp));
-
-			const TArray<FXmlNode*>& ChangeNodes = ChangesNode->GetChildrenNodes();
-			for (const FXmlNode* ChangeNode : ChangeNodes)
-			{
-				check(ChangeNode);
-				const FXmlNode* PathNode = ChangeNode->FindChildNode(Path);
-				const FXmlNode* TypeNode = ChangeNode->FindChildNode(Type);
-				if (PathNode == nullptr || TypeNode == nullptr)
-				{
-					continue;
-				}
-
-				// Here we make sure to only collect file states, since we shouldn't display the added directories to the Editor
-				FString FileName = PathNode->GetContent();
-				int32 DotIndex;
-				if (FileName.FindChar(TEXT('.'), DotIndex))
-				{
-					FPlasticSourceControlState FileState(FPaths::ConvertRelativePathToFull(WorkingDirectory, MoveTemp(FileName)));
-					FileState.WorkspaceState = PlasticSourceControlUtils::StateFromPlasticStatus(TypeNode->GetContent());
-					FileState.Changelist = ChangelistState.Changelist;
-					OutCLFilesStates[ChangelistIndex].Add(MoveTemp(FileState));
-				}
-			}
-
-			OutChangelistsStates.Add(ChangelistState);
-		}
-	}
-
-	if (!OutChangelistsStates.FindByPredicate(
-		[](const FPlasticSourceControlChangelistState& CLState) { return CLState.Changelist.IsDefault(); }
-	))
-	{
-		// No Default Changelists isn't an error, but the Editor UX expects to always the Default changelist (so you can always move files back to it)
-		FPlasticSourceControlChangelistState DefaultChangelistState(FPlasticSourceControlChangelist::DefaultChangelist);
-		OutChangelistsStates.Insert(DefaultChangelistState, 0);
-		OutCLFilesStates.Insert(TArray<FPlasticSourceControlState>(), 0);
-	}
-
-	return true;
-}
-
 // Run a Plastic "status --changelist --xml" and parse its XML result.
 bool RunGetChangelists(TArray<FPlasticSourceControlChangelistState>& OutChangelistsStates, TArray<TArray<FPlasticSourceControlState>>& OutCLFilesStates, TArray<FString>& OutErrorMessages)
 {
@@ -1679,19 +797,7 @@ bool RunGetChangelists(TArray<FPlasticSourceControlChangelistState>& OutChangeli
 	bool bResult = RunCommand(TEXT("status"), Parameters, TArray<FString>(), Results, Errors);
 	if (bResult)
 	{
-		FXmlFile XmlFile;
-		{
-			TRACE_CPUPROFILER_EVENT_SCOPE(PlasticSourceControlUtils::RunGetChangelists::FXmlFile::LoadFile);
-			bResult = XmlFile.LoadFile(Results, EConstructMethod::ConstructFromBuffer);
-		}
-		if (bResult)
-		{
-			bResult = ParseChangelistsResults(XmlFile, OutChangelistsStates, OutCLFilesStates);
-		}
-		else
-		{
-			UE_LOG(LogSourceControl, Error, TEXT("RunGetChangelists: XML parse error '%s'"), *XmlFile.GetLastError())
-		}
+		bResult = PlasticSourceControlParsers::ParseChangelistsResults(Results, OutChangelistsStates, OutCLFilesStates);
 	}
 	if (!Errors.IsEmpty())
 	{
@@ -1701,35 +807,10 @@ bool RunGetChangelists(TArray<FPlasticSourceControlChangelistState>& OutChangeli
 	return bResult;
 }
 
-// Parse the one letter file status in front of each line of the 'cm diff sh:<ShelveId>'
-EWorkspaceState ParseShelveFileStatus(const TCHAR InFileStatus)
-{
-	if (InFileStatus == 'A') // Added
-	{
-		return EWorkspaceState::Added;
-	}
-	else if (InFileStatus == 'D') // Deleted
-	{
-		return EWorkspaceState::Deleted;
-	}
-	else if (InFileStatus == 'C') // Changed (CheckedOut or not)
-	{
-		return EWorkspaceState::CheckedOut;
-	}
-	else if (InFileStatus == 'M') // Moved/Renamed (or Locally Moved)
-	{
-		return EWorkspaceState::Moved;
-	}
-	else
-	{
-		UE_LOG(LogSourceControl, Warning, TEXT("Unknown file status '%c'"), InFileStatus);
-		return EWorkspaceState::Unknown;
-	}
-}
-
-void AddShelvedFileToChangelist(FPlasticSourceControlChangelistState& InOutChangelistsState, FString&& InFilename, EWorkspaceState InShelveStatus)
+void AddShelvedFileToChangelist(FPlasticSourceControlChangelistState& InOutChangelistsState, FString&& InFilename, EWorkspaceState InShelveStatus, FString&& InMovedFrom)
 {
 	TSharedRef<FPlasticSourceControlState, ESPMode::ThreadSafe> ShelveState = MakeShared<FPlasticSourceControlState>(MoveTemp(InFilename), InShelveStatus);
+	ShelveState->MovedFrom = MoveTemp(InMovedFrom);
 
 	// Add one revision to be able to fetch the shelved file for diff, if it's not marked for deletion.
 	if (InShelveStatus != EWorkspaceState::Deleted)
@@ -1739,7 +820,7 @@ void AddShelvedFileToChangelist(FPlasticSourceControlChangelistState& InOutChang
 		SourceControlRevision->Filename = ShelveState->GetFilename();
 		SourceControlRevision->ShelveId = InOutChangelistsState.ShelveId;
 		SourceControlRevision->ChangesetNumber = InOutChangelistsState.ShelveId; // Note: for display in the diff window only
-		SourceControlRevision->Date = InOutChangelistsState.ShelveDate; // Note: not yet used for display as of UE5.1.1
+		SourceControlRevision->Date = InOutChangelistsState.ShelveDate; // Note: not yet used for display as of UE5.2
 
 		ShelveState->History.Add(SourceControlRevision);
 	}
@@ -1759,52 +840,6 @@ void AddShelvedFileToChangelist(FPlasticSourceControlChangelistState& InOutChang
 	}
 }
 
-
-/**
- * Parse results of the 'cm diff sh:<ShelveId>' command.
- *
- * Results of the diff command looks like that:
-C "Content\NewFolder\BP_CheckedOut.uasset"
-C "Content\NewFolder\BP_Renamed.uasset"
-A "Content\NewFolder\BP_ControlledUnchanged.uasset"
-D "Content\NewFolder\BP_Changed.uasset"
-M "Content\NewFolder\BP_ControlledUnchanged.uasset" "Content\NewFolder\BP_Renamed.uasset"
-*/
-bool ParseShelveDiffResults(const FString InWorkingDirectory, TArray<FString>&& InResults, FPlasticSourceControlChangelistState& InOutChangelistsState)
-{
-	bool bCommandSuccessful = true;
-
-	InOutChangelistsState.ShelvedFiles.Reset(InResults.Num());
-	for (FString& Result : InResults)
-	{
-		EWorkspaceState ShelveStatus = ParseShelveFileStatus(Result[0]);
-		
-		// Remove outer double quotes
-		Result.MidInline(3, Result.Len() - 4, EAllowShrinking::No);
-		if (ShelveStatus == EWorkspaceState::Moved)
-		{
-			// Search for the inner double quotes in the middle of "Content/Source.uasset" "Content/Destination.uasset" to keep only the destination filename
-			int32 RenameIndex;
-			if (Result.FindLastChar(TEXT('"'), RenameIndex))
-			{
-				Result.RightChopInline(RenameIndex + 1);
-			}
-		}
-		
-		if (ShelveStatus != EWorkspaceState::Unknown && !Result.IsEmpty())
-		{
-			FString AbsoluteFilename = FPaths::ConvertRelativePathToFull(InWorkingDirectory, MoveTemp(Result));
-			AddShelvedFileToChangelist(InOutChangelistsState, MoveTemp(AbsoluteFilename), ShelveStatus);
-		}
-		else
-		{
-			bCommandSuccessful = false;
-		}
-	}
-
-	return bCommandSuccessful;
-}
-
 /**
  * Run for each shelve a "diff sh:<ShelveId>" and parse their result to list their files.
  * @param	InOutChangelistsStates	The list of changelists, filled with their shelved files
@@ -1814,7 +849,7 @@ bool RunGetShelveFiles(TArray<FPlasticSourceControlChangelistState>& InOutChange
 {
 	bool bCommandSuccessful = true;
 
-	const FString& WorkingDirectory = FPlasticSourceControlModule::Get().GetProvider().GetPathToWorkspaceRoot();
+	const FString& WorkspaceRoot = FPlasticSourceControlModule::Get().GetProvider().GetPathToWorkspaceRoot();
 
 	for (FPlasticSourceControlChangelistState& ChangelistState : InOutChangelistsStates)
 	{
@@ -1822,11 +857,12 @@ bool RunGetShelveFiles(TArray<FPlasticSourceControlChangelistState>& InOutChange
 		{
 			TArray<FString> Results;
 			TArray<FString> Parameters;
+			// TODO switch to custom format --format="{status};{path};{srccmpath}" for better parsing, and perhaps reusing code
 			Parameters.Add(FString::Printf(TEXT("sh:%d"), ChangelistState.ShelveId));
 			const bool bDiffSuccessful = PlasticSourceControlUtils::RunCommand(TEXT("diff"), Parameters, TArray<FString>(), Results, OutErrorMessages);
 			if (bDiffSuccessful)
 			{
-				bCommandSuccessful = ParseShelveDiffResults(WorkingDirectory, MoveTemp(Results), ChangelistState);
+				bCommandSuccessful = PlasticSourceControlParsers::ParseShelveDiffResult(WorkspaceRoot, MoveTemp(Results), ChangelistState);
 			}
 		}
 	}
@@ -1834,103 +870,79 @@ bool RunGetShelveFiles(TArray<FPlasticSourceControlChangelistState>& InOutChange
 	return bCommandSuccessful;
 }
 
-/**
- * Parse results of the 'cm find "shelves where owner='me'" --xml --encoding="utf-8"' command.
- *
- * Find shelves with comments starting like "ChangelistXXX: " and matching an existing Changelist number XXX
- *
- * Results of the find command looks like the following; note the "Changelist67: " prefix of the comment:
-<?xml version="1.0" encoding="utf-8" ?>
-<PLASTICQUERY>
-  <SHELVE>
-	<ID>1376</ID>
-	<SHELVEID>9</SHELVEID>
-    <COMMENT>Changelist67: test by Sebastien</COMMENT>
-	<DATE>2022-06-30T16:39:55+02:00</DATE>
-    <OWNER>[REDACTED_EMAIL]</OWNER>
-    <REPOSITORY>UE5PlasticPluginDev</REPOSITORY>
-    <REPNAME>UE5PlasticPluginDev</REPNAME>
-    <REPSERVER>test@cloud</REPSERVER>
-	<PARENT>45</PARENT>
-	<GUID>8fbefbcc-81a7-4b81-9b99-b51f4873d09f</GUID>
-  </SHELVE>
-  [...]
-</PLASTICQUERY>
-*/
-static bool ParseShelvesResults(const FXmlFile& InXmlResult, TArray<FPlasticSourceControlChangelistState>& InOutChangelistsStates)
+// Run find "shelves where owner='me'" and for each shelve matching a changelist a "diff sh:<ShelveId>" and parse their results.
+bool RunGetShelves(TArray<FPlasticSourceControlChangelistState>& InOutChangelistsStates, TArray<FString>& OutErrorMessages)
 {
-	static const FString PlasticQuery(TEXT("PLASTICQUERY"));
-	static const FString Shelve(TEXT("SHELVE"));
-	static const FString ShelveId(TEXT("SHELVEID"));
-	static const FString Date(TEXT("DATE"));
-	static const FString Comment(TEXT("COMMENT"));
-
-	const FString& WorkingDirectory = FPlasticSourceControlModule::Get().GetProvider().GetPathToWorkspaceRoot();
-
-	const FXmlNode* PlasticQueryNode = InXmlResult.GetRootNode();
-	if (PlasticQueryNode == nullptr || PlasticQueryNode->GetTag() != PlasticQuery)
-	{
-		return false;
-	}
+	bool bCommandSuccessful;
 
-	const TArray<FXmlNode*>& ShelvesNodes = PlasticQueryNode->GetChildrenNodes();
-	for (const FXmlNode* ShelveNode : ShelvesNodes)
+	FString Results;
+	FString Errors;
+	TArray<FString> Parameters;
+	Parameters.Add(TEXT("\"shelves where owner = 'me'\""));
+	Parameters.Add(TEXT("--xml"));
+	Parameters.Add(TEXT("--encoding=\"utf-8\""));
+	bCommandSuccessful = PlasticSourceControlUtils::RunCommand(TEXT("find"), Parameters, TArray<FString>(), Results, Errors);
+	if (bCommandSuccessful)
 	{
-		check(ShelveNode);
-		const FXmlNode* ShelveIdNode = ShelveNode->FindChildNode(ShelveId);
-		const FXmlNode* CommentNode = ShelveNode->FindChildNode(Comment);
-		if (ShelveIdNode == nullptr || CommentNode == nullptr)
+		bCommandSuccessful = PlasticSourceControlParsers::ParseShelvesResults(Results, InOutChangelistsStates);
+		if (bCommandSuccessful)
 		{
-			continue;
+			bCommandSuccessful = RunGetShelveFiles(InOutChangelistsStates, OutErrorMessages);
 		}
+	}
+	if (!Errors.IsEmpty())
+	{
+		OutErrorMessages.Add(MoveTemp(Errors));
+	}
 
-		const FString& ShelveIdString = ShelveIdNode->GetContent();
-		const FString& CommentString = CommentNode->GetContent();
+	return bCommandSuccessful;
+}
 
-		// Search if there is a changelist matching the shelve (that is, a shelve with a comment starting with "ChangelistXXX: ")
-		for (FPlasticSourceControlChangelistState& ChangelistState : InOutChangelistsStates)
-		{
-			FPlasticSourceControlChangelistRef Changelist = StaticCastSharedRef<FPlasticSourceControlChangelist>(ChangelistState.GetChangelist());
-			const FString ChangelistPrefix = FString::Printf(TEXT("Changelist%s: "), *Changelist->GetName());
-			if (CommentString.StartsWith(ChangelistPrefix))
-			{
-				ChangelistState.ShelveId = FCString::Atoi(*ShelveIdString);
+/**
+ * Run for a shelve a "diff sh:<ShelveId> --format='{status};{baserevid};{path}'" and parse the result to list its files.
+ * @param	InOutChangelistsStates	The list of changelists, filled with their shelved files
+ * @param	OutErrorMessages		Any errors (from StdErr) as an array per-line
+ */
+bool RunGetShelveFiles(const int32 InShelveId, TArray<FPlasticSourceControlRevision>& OutBaseRevisions, TArray<FString>& OutErrorMessages)
+{
+	bool bCommandSuccessful = true;
 
-				if (const FXmlNode* DateNode = ShelveNode->FindChildNode(Date))
-				{
-					const FString& DateIso = DateNode->GetContent();
-					FDateTime::ParseIso8601(*DateIso, ChangelistState.ShelveDate);
-				}
-			}
+	const FString& WorkspaceRoot = FPlasticSourceControlModule::Get().GetProvider().GetPathToWorkspaceRoot();
+
+	if (InShelveId != ISourceControlState::INVALID_REVISION)
+	{
+		TArray<FString> Results;
+		TArray<FString> Parameters;
+		Parameters.Add(FString::Printf(TEXT("sh:%d"), InShelveId));
+		Parameters.Add(TEXT("--format=\"{status};{baserevid};{path}\""));
+		Parameters.Add(TEXT("--encoding=\"utf-8\""));
+		const bool bDiffSuccessful = PlasticSourceControlUtils::RunCommand(TEXT("diff"), Parameters, TArray<FString>(), Results, OutErrorMessages);
+		if (bDiffSuccessful)
+		{
+			bCommandSuccessful = PlasticSourceControlParsers::ParseShelveDiffResults(WorkspaceRoot, MoveTemp(Results), OutBaseRevisions);
 		}
 	}
 
-	return true;
+	return bCommandSuccessful;
 }
 
-// Run find "shelves where owner='me'" and for each shelve matching a changelist a "diff sh:<ShelveId>" and parse their results.
-bool RunGetShelves(TArray<FPlasticSourceControlChangelistState>& InOutChangelistsStates, TArray<FString>& OutErrorMessages)
+bool RunGetShelve(const int32 InShelveId, FString& OutComment, FDateTime& OutDate, FString& OutOwner, TArray<FPlasticSourceControlRevision>& OutBaseRevisions, TArray<FString>& OutErrorMessages)
 {
 	bool bCommandSuccessful;
 
 	FString Results;
 	FString Errors;
 	TArray<FString> Parameters;
-	Parameters.Add(TEXT("\"shelves where owner = 'me'\""));
+	Parameters.Add(FString::Printf(TEXT("\"shelves where ShelveId = %d\""), InShelveId));
 	Parameters.Add(TEXT("--xml"));
 	Parameters.Add(TEXT("--encoding=\"utf-8\""));
 	bCommandSuccessful = PlasticSourceControlUtils::RunCommand(TEXT("find"), Parameters, TArray<FString>(), Results, Errors);
 	if (bCommandSuccessful)
 	{
-		FXmlFile XmlFile;
-		bCommandSuccessful = XmlFile.LoadFile(Results, EConstructMethod::ConstructFromBuffer);
+		bCommandSuccessful = PlasticSourceControlParsers::ParseShelvesResult(Results, OutComment, OutDate, OutOwner);
 		if (bCommandSuccessful)
 		{
-			bCommandSuccessful = ParseShelvesResults(XmlFile, InOutChangelistsStates);
-			if (bCommandSuccessful)
-			{
-				bCommandSuccessful = RunGetShelveFiles(InOutChangelistsStates, OutErrorMessages);
-			}
+			bCommandSuccessful = RunGetShelveFiles(InShelveId, OutBaseRevisions, OutErrorMessages);
 		}
 	}
 	if (!Errors.IsEmpty())
@@ -1941,44 +953,142 @@ bool RunGetShelves(TArray<FPlasticSourceControlChangelistState>& InOutChangelist
 	return bCommandSuccessful;
 }
 
-bool UpdateCachedStates(TArray<FPlasticSourceControlState>&& InStates)
+bool RunGetBranches(const FDateTime& InFromDate, TArray<FPlasticSourceControlBranchRef>& OutBranches, TArray<FString>& OutErrorMessages)
 {
-	FPlasticSourceControlProvider& Provider = FPlasticSourceControlModule::Get().GetProvider();
-	const FDateTime Now = FDateTime::Now();
+	bool bCommandSuccessful;
 
-	for (auto&& InState : InStates)
+	FString Results;
+	FString Errors;
+	TArray<FString> Parameters;
+	if (InFromDate != FDateTime())
 	{
-		TSharedRef<FPlasticSourceControlState, ESPMode::ThreadSafe> State = Provider.GetStateInternal(InState.LocalFilename);
-		*State = MoveTemp(InState);
-		State->TimeStamp = Now;
+		// Find branches created since this date or containing changes dating from or after this date
+		Parameters.Add(FString::Printf(TEXT("\"branches where date >= '%d/%d/%d'\" or changesets >= '%d/%d/%d'"),
+			InFromDate.GetYear(), InFromDate.GetMonth(), InFromDate.GetDay(),
+			InFromDate.GetYear(), InFromDate.GetMonth(), InFromDate.GetDay()
+		));
+	}
+	else
+	{
+		Parameters.Add(TEXT("\"branches\""));
+	}
+	Parameters.Add(TEXT("--xml"));
+	Parameters.Add(TEXT("--encoding=\"utf-8\""));
+	bCommandSuccessful = PlasticSourceControlUtils::RunCommand(TEXT("find"), Parameters, TArray<FString>(), Results, Errors);
+	if (bCommandSuccessful)
+	{
+		bCommandSuccessful = PlasticSourceControlParsers::ParseBranchesResults(Results, OutBranches);
+	}
+	if (!Errors.IsEmpty())
+	{
+		OutErrorMessages.Add(MoveTemp(Errors));
 	}
 
-	return (InStates.Num() > 0);
+	return bCommandSuccessful;
 }
 
-/**
- * Helper struct for RemoveRedundantErrors()
- */
-struct FRemoveRedundantErrors
+bool RunSwitchToBranch(const FString& InBranchName, TArray<FString>& OutUpdatedFiles, TArray<FString>& OutErrorMessages)
 {
-	explicit FRemoveRedundantErrors(const FString& InFilter)
-		: Filter(InFilter)
+	bool bResult = false;
+
+	const FScopedTempFile SwitchResultFile;
+	TArray<FString> InfoMessages;
+	TArray<FString> Parameters;
+	Parameters.Add(FString::Printf(TEXT("--xml=\"%s\""), *SwitchResultFile.GetFilename()));
+	Parameters.Add(TEXT("--encoding=\"utf-8\""));
+	Parameters.Add(FString::Printf(TEXT("\"br:%s\""), *InBranchName));
+	bResult = PlasticSourceControlUtils::RunCommand(TEXT("switch"), Parameters, TArray<FString>(), InfoMessages, OutErrorMessages);
+	if (bResult)
 	{
+		// Load and parse the result of the update command
+		FString Results;
+		if (FFileHelper::LoadFileToString(Results, *SwitchResultFile.GetFilename()))
+		{
+			bResult = PlasticSourceControlParsers::ParseUpdateResults(Results, OutUpdatedFiles);
+		}
 	}
 
-	bool operator()(const FString& String) const
+	return bResult;
+}
+
+bool RunMergeBranch(const FString& InBranchName, TArray<FString>& OutUpdatedFiles, TArray<FString>& OutErrorMessages)
+{
+	bool bResult = false;
+
+	const FScopedTempFile MergeResultFile;
+	TArray<FString> InfoMessages;
+	TArray<FString> Parameters;
+	Parameters.Add(FString::Printf(TEXT("--xml=\"%s\""), *MergeResultFile.GetFilename()));
+	Parameters.Add(TEXT("--encoding=\"utf-8\""));
+	Parameters.Add(TEXT("--merge"));
+	Parameters.Add(FString::Printf(TEXT("\"br:%s\""), *InBranchName));
+	bResult = PlasticSourceControlUtils::RunCommand(TEXT("merge"), Parameters, TArray<FString>(), InfoMessages, OutErrorMessages);
+	if (bResult)
 	{
-		if (String.Contains(Filter))
+		// Load and parse the result of the update command
+		FString Results;
+		if (FFileHelper::LoadFileToString(Results, *MergeResultFile.GetFilename()))
 		{
-			return true;
+			PlasticSourceControlParsers::ParseMergeResults(Results, OutUpdatedFiles);
 		}
+	}
 
-		return false;
+	return bResult;
+}
+
+bool RunCreateBranch(const FString& InBranchName, const FString& InComment, TArray<FString>& OutErrorMessages)
+{
+	// make a temp file to place our comment message in
+	const FScopedTempFile BranchCommentFile(InComment);
+	if (!BranchCommentFile.GetFilename().IsEmpty())
+	{
+		TArray<FString> Parameters;
+		TArray<FString> InfoMessages;
+		Parameters.Add(TEXT("create"));
+		Parameters.Add(FString::Printf(TEXT("\"%s\""), *InBranchName));
+		Parameters.Add(FString::Printf(TEXT("--commentsfile=\"%s\""), *BranchCommentFile.GetFilename()));
+		return PlasticSourceControlUtils::RunCommand(TEXT("branch"), Parameters, TArray<FString>(), InfoMessages, OutErrorMessages);
 	}
 
-	/** The filter string we try to identify in the reported error */
-	FString Filter;
-};
+	return false;
+}
+
+bool RunRenameBranch(const FString& InOldName, const FString& InNewName, TArray<FString>& OutErrorMessages)
+{
+	TArray<FString> Parameters;
+	TArray<FString> InfoMessages;
+	Parameters.Add(TEXT("rename"));
+	Parameters.Add(FString::Printf(TEXT("\"br:%s\""), *InOldName));
+	Parameters.Add(FString::Printf(TEXT("\"%s\""), *InNewName));
+	return PlasticSourceControlUtils::RunCommand(TEXT("branch"), Parameters, TArray<FString>(), InfoMessages, OutErrorMessages);
+}
+
+bool RunDeleteBranches(const TArray<FString>& InBranchNames, TArray<FString>& OutErrorMessages)
+{
+	TArray<FString> Parameters;
+	TArray<FString> InfoMessages;
+	Parameters.Add(TEXT("delete"));
+	for (const FString& BranchName : InBranchNames)
+	{
+		Parameters.Add(FString::Printf(TEXT("\"br:%s\""), *BranchName));
+	}
+	return PlasticSourceControlUtils::RunCommand(TEXT("branch"), Parameters, TArray<FString>(), InfoMessages, OutErrorMessages);
+}
+
+bool UpdateCachedStates(TArray<FPlasticSourceControlState>&& InStates)
+{
+	FPlasticSourceControlProvider& Provider = FPlasticSourceControlModule::Get().GetProvider();
+	const FDateTime Now = FDateTime::Now();
+
+	for (auto&& InState : InStates)
+	{
+		TSharedRef<FPlasticSourceControlState, ESPMode::ThreadSafe> State = Provider.GetStateInternal(InState.LocalFilename);
+		*State = MoveTemp(InState);
+		State->TimeStamp = Now;
+	}
+
+	return (InStates.Num() > 0);
+}
 
 void RemoveRedundantErrors(FPlasticSourceControlCommand& InCommand, const FString& InFilter)
 {
@@ -1992,7 +1102,7 @@ void RemoveRedundantErrors(FPlasticSourceControlCommand& InCommand, const FStrin
 		}
 	}
 
-	InCommand.ErrorMessages.RemoveAll(FRemoveRedundantErrors(InFilter));
+	InCommand.ErrorMessages.RemoveAll(PlasticSourceControlParsers::FRemoveRedundantErrors(InFilter));
 
 	// if we have no error messages now, assume success!
 	if (bFoundRedundantError && InCommand.ErrorMessages.Num() == 0 && !InCommand.bCommandSuccessful)
diff --git a/Engine/Plugins/Developer/PlasticSourceControl/Source/PlasticSourceControl/Private/PlasticSourceControlUtils.h b/Engine/Plugins/Developer/PlasticSourceControl/Source/PlasticSourceControl/Private/PlasticSourceControlUtils.h
--- a/Engine/Plugins/Developer/PlasticSourceControl/Source/PlasticSourceControl/Private/PlasticSourceControlUtils.h
+++ b/Engine/Plugins/Developer/PlasticSourceControl/Source/PlasticSourceControl/Private/PlasticSourceControlUtils.h
@@ -3,19 +3,48 @@
 #pragma once
 
 #include "CoreMinimal.h"
+
 #include "PlasticSourceControlRevision.h"
 
 class FPlasticSourceControlChangelistState;
 class FPlasticSourceControlCommand;
 class FPlasticSourceControlState;
 struct FSoftwareVersion;
+typedef TSharedRef<class FPlasticSourceControlBranch, ESPMode::ThreadSafe> FPlasticSourceControlBranchRef;
+
+namespace PlasticSourceControlParsers
+{
+class FSmartLockInfoParser;
+} // namespace PlasticSourceControlParsers
 
 enum class EWorkspaceState;
 
 namespace PlasticSourceControlUtils
 {
 
-const struct FSoftwareVersion& GetOldestSupportedPlasticScmVersion();
+/**
+ * Run a Plastic command - the result is the output of cm, as a multi-line string.
+ *
+ * @param	InCommand			The Plastic command - e.g. commit
+ * @param	InParameters		The parameters to the Plastic command
+ * @param	InFiles				The files to be operated on
+ * @param	OutResults			The results (from StdOut) as a multi-line string.
+ * @param	OutErrors			Any errors (from StdErr) as a multi-line string.
+ * @returns true if the command succeeded and returned no errors
+ */
+bool RunCommand(const FString& InCommand, const TArray<FString>& InParameters, const TArray<FString>& InFiles, FString& OutResults, FString& OutErrors);
+
+/**
+ * Run a Plastic command - the result is parsed in an array of strings.
+ *
+ * @param	InCommand			The Plastic command - e.g. commit
+ * @param	InParameters		The parameters to the Plastic command
+ * @param	InFiles				The files to be operated on
+ * @param	OutResults			The results (from StdOut) as an array per-line
+ * @param	OutErrorMessages	Any errors (from StdErr) as an array per-line
+ * @returns true if the command succeeded and returned no errors
+ */
+bool RunCommand(const FString& InCommand, const TArray<FString>& InParameters, const TArray<FString>& InFiles, TArray<FString>& OutResults, TArray<FString>& OutErrorMessages);
 
 /**
  * Find the path to the Plastic binary: for now relying on the Path to access the "cm" command.
@@ -31,33 +60,39 @@ FString FindPlasticBinaryPath();
 bool GetWorkspacePath(const FString& InPathToGameDir, FString& OutWorkspaceRoot);
 
 /**
- * Get Plastic SCM CLI version
- * @param	OutCliVersion		Version of the Plastic SCM Command Line Interface tool
+ * Get Unity Version Control CLI version
+ * @param	OutCliVersion		Version of the Unity Version Control Command Line Interface tool
  * @returns true if the command succeeded
 */
 bool GetPlasticScmVersion(FSoftwareVersion& OutPlasticScmVersion);
 
 /**
- * Get Plastic SCM CLI location
+ * Get Unity Version Control CLI location
  * @param	OutCmLocation		Path to the "cm" executable
  * @returns true if the command succeeded
 */
 bool GetCmLocation(FString& OutCmLocation);
 
 /**
- * Checks weather Plastic SCM is configured to set files as read-only on update & checkin
+ * Checks weather Unity Version Control is configured to set files as read-only on update & checkin
  * @returns true if SetFilesAsReadOnly is enabled in client.conf
 */
 bool GetConfigSetFilesAsReadOnly();
 
 /**
- * Get Plastic SCM current user
- * @param	OutUserName			Name of the Plastic SCM user configured globally
+ * Get from config the name of the default cloud organization or the url of the default on-prem server
+ * @returns Returns the location of the default repository server
+*/
+FString GetConfigDefaultRepServer();
+
+/**
+ * Get Unity Version Control current user
+ * @param	OutUserName			Name of the Unity Version Control user configured globally
  */
 void GetUserName(FString& OutUserName);
 
 /**
- * Get Plastic workspace name
+ * Get workspace name
  * @param	InWorkspaceRoot		The workspace from where to run the command - usually the Game directory
  * @param	OutWorkspaceName	Name of the current workspace
  * @param	OutErrorMessages	Any errors (from StdErr) as an array per-line
@@ -65,64 +100,56 @@ void GetUserName(FString& OutUserName);
 bool GetWorkspaceName(const FString& InWorkspaceRoot, FString& OutWorkspaceName, TArray<FString>& OutErrorMessages);
 
 /**
- * Get Plastic repository name and server URL, branch name and current changeset number
- * @param	OutChangeset		The current Changeset Number
+ * Get workspace info: the current branch, repository name, and server URL
+ * 
+ * @param	OutBranchName		Name of the current branch
  * @param	OutRepositoryName	Name of the repository of the current workspace
  * @param	OutServerUrl		URL/Port of the server of the repository
- * @param	OutBranchName		Name of the current checked-out branch
  * @param	OutErrorMessages	Any errors (from StdErr) as an array per-line
  */
-bool GetWorkspaceInformation(int32& OutChangeset, FString& OutRepositoryName, FString& OutServerUrl, FString& OutBranchName, TArray<FString>& OutErrorMessages);
+bool GetWorkspaceInfo(FString& OutBranchName, FString& OutRepositoryName, FString& OutServerUrl, TArray<FString>& OutErrorMessages);
 
 /**
- * Get Plastic repository name, server URL, and branch name
+ * Get workspace info and check the connection to the server
  *
- * Note: this is a local fast variant, not making network call to the server.
- * 
- * @param	OutBranchName		Name of the current checked-out branch
+ * @param	OutBranchName		Name of the current branch
  * @param	OutRepositoryName	Name of the repository of the current workspace
  * @param	OutServerUrl		URL/Port of the server of the repository
+ * @param	OutInfoMessages		Result of the connection test
  * @param	OutErrorMessages	Any errors (from StdErr) as an array per-line
  */
-bool GetWorkspaceInfo(FString& OutBranchName, FString& OutRepositoryName, FString& OutServerUrl, TArray<FString>& OutErrorMessages);
+bool RunCheckConnection(FString& OutBranchName, FString& OutRepositoryName, FString& OutServerUrl, TArray<FString>& OutInfoMessages, TArray<FString>& OutErrorMessages);
 
 /**
- * Use the Project Settings to replace Plastic SCM full username/e-mail by a shorter version for display.
+ * Use the Project Settings to replace Unity Version Control full username/e-mail by a shorter version for display.
  *
  * Used when retrieving the username of a revision, to display in history and content browser asset tooltip.
  *
- * @param	InUserName			The Plastic SCM username to shorten for display.
+ * @param	InUserName			The Unity Version Control username to shorten for display.
  */
 FString UserNameToDisplayName(const FString& InUserName);
 
 /**
- * Run a Plastic command - the result is the output of cm, as a multi-line string.
+ * Run a Plastic "lock list" command and parse it.
  *
- * @param	InCommand			The Plastic command - e.g. commit
- * @param	InParameters		The parameters to the Plastic command
- * @param	InFiles				The files to be operated on
- * @param	OutResults			The results (from StdOut) as a multi-line string.
- * @param	OutErrors			Any errors (from StdErr) as a multi-line string.
+ * @param	InRepository		The repository to ask for the locks
+ * @param	OutSmartLocks		The list of smart locks
  * @returns true if the command succeeded and returned no errors
  */
-bool RunCommand(const FString& InCommand, const TArray<FString>& InParameters, const TArray<FString>& InFiles, FString& OutResults, FString& OutErrors);
+bool RunListSmartLocks(const FString& InRepository, TMap<FString, PlasticSourceControlParsers::FSmartLockInfoParser>& OutSmartLocks);
 
-/**
- * Run a Plastic command - the result is parsed in an array of strings.
- *
- * @param	InCommand			The Plastic command - e.g. commit
- * @param	InParameters		The parameters to the Plastic command
- * @param	InFiles				The files to be operated on
- * @param	OutResults			The results (from StdOut) as an array per-line
- * @param	OutErrorMessages	Any errors (from StdErr) as an array per-line
- * @returns true if the command succeeded and returned no errors
- */
-bool RunCommand(const FString& InCommand, const TArray<FString>& InParameters, const TArray<FString>& InFiles, TArray<FString>& OutResults, TArray<FString>& OutErrorMessages);
+// Specify the "search type" for the "status" command
+enum class EStatusSearchType
+{
+	All,			// status --all --ignored (this can take much longer, searching for local changes, especially on the first call)
+	ControlledOnly	// status --controlledchanged
+};
 
 /**
  * Run a Plastic "status" command and parse it.
  *
  * @param	InFiles				The files to be operated on
+ * @param	InSearchType		Call "status" with "--all", or with just "--controlledchanged" when doing only a quick check following a source control operation
  * @param	bInUpdateHistory	If getting the history of files, force execute the fileinfo command required to do get RepSpec of xlinks (history view or visual diff)
  * @param	OutErrorMessages	Any errors (from StdErr) as an array per-line
  * @param	OutStates			States of the files
@@ -130,17 +157,16 @@ bool RunCommand(const FString& InCommand, const TArray<FString>& InParameters, c
  * @param	OutBranchName		Name of the current checked-out branch
  * @returns true if the command succeeded and returned no errors
  */
-bool RunUpdateStatus(const TArray<FString>& InFiles, const bool bInUpdateHistory, TArray<FString>& OutErrorMessages, TArray<FPlasticSourceControlState>& OutStates, int32& OutChangeset, FString& OutBranchName);
+bool RunUpdateStatus(const TArray<FString>& InFiles, const EStatusSearchType InSearchType, const bool bInUpdateHistory, TArray<FString>& OutErrorMessages, TArray<FPlasticSourceControlState>& OutStates, int32& OutChangeset, FString& OutBranchName);
 
 /**
  * Run a Plastic "cat" command to dump the binary content of a revision into a file.
  *
- * @param	InPathToPlasticBinary	The path to the Plastic binary
  * @param	InRevSpec				The revision specification to get
  * @param	InDumpFileName			The temporary file to dump the revision
  * @returns true if the command succeeded and returned no errors
 */
-bool RunDumpToFile(const FString& InPathToPlasticBinary, const FString& InRevSpec, const FString& InDumpFileName);
+bool RunGetFile(const FString& InRevSpec, const FString& InDumpFileName);
 
 /**
  * Run Plastic "history" and "log" commands and parse their XML results.
@@ -176,13 +202,72 @@ bool RunGetChangelists(TArray<FPlasticSourceControlChangelistState>& OutChangeli
  */
 bool RunGetShelves(TArray<FPlasticSourceControlChangelistState>& InOutChangelistsStates, TArray<FString>& OutErrorMessages);
 
+/**
+ * Run find "shelves where ShelveId='NNN'" and a "diff sh:<ShelveId>" and parse their results.
+ * @param	InShelveId			Shelve Id
+ * @param	OutComment			Shelve Comment
+ * @param	OutDate				Shelve Date
+ * @param	OutOwner			Shelve Owner
+ * @param	OutStates			Files in the shelve and their base revision id
+ * @param	OutErrorMessages	Any errors (from StdErr) as an array per-line
+ */
+bool RunGetShelve(const int32 InShelveId, FString& OutComment, FDateTime& OutDate, FString& OutOwner, TArray<FPlasticSourceControlRevision>& OutBaseRevisions, TArray<FString>& OutErrorMessages);
+
 /**
  * Add a file to the shelve associated with a changelist.
  * @param	InOutChangelistsState	The changelist to add the file to
  * @param	InFilename				The file to add to the shelve
  * @param	InShelveStatus			The status of the file
+ * @param	InMovedFrom				If moved, the original filename
+*/
+void AddShelvedFileToChangelist(FPlasticSourceControlChangelistState& InOutChangelistsState, FString&& InFilename, EWorkspaceState InShelveStatus, FString&& InMovedFrom);
+
+/**
+ * Run find "branches where date >= 'YYYY-MM-DD' or changesets >= 'YYYY-MM-DD'" and parse the results.
+ * @param	InFromDate				The date to search from
+ * @param	OutBranches				The list of branches
+ * @param	OutErrorMessages		Any errors (from StdErr) as an array per-line
+ */
+bool RunGetBranches(const FDateTime& InFromDate, TArray<FPlasticSourceControlBranchRef>& OutBranches, TArray<FString>& OutErrorMessages);
+
+/**
+ * Run switch br:/name and parse the results.
+ * @param	InBranchName			The name of the branch to switch the workspace to
+ * @param	OutUpdatedFiles			The files that where updated
+ * @param	OutErrorMessages		Any errors (from StdErr) as an array per-line
+ */
+bool RunSwitchToBranch(const FString& InBranchName, TArray<FString>& OutUpdatedFiles, TArray<FString>& OutErrorMessages);
+
+/**
+ * Run merge br:/name and parse the results.
+ * @param	InBranchName			The name of the branch to merge to the current branch
+ * @param	OutUpdatedFiles			The files that where updated
+ * @param	OutErrorMessages		Any errors (from StdErr) as an array per-line
+ */
+bool RunMergeBranch(const FString& InBranchName, TArray<FString>& OutUpdatedFiles, TArray<FString>& OutErrorMessages);
+
+/**
+ * Run branch create <name> --commentsfile
+ * @param	InBranchName			The name of the branch to create
+ * @param	InComment				The comment for the new branch to create
+ * @param	OutErrorMessages		Any errors (from StdErr) as an array per-line
+ */
+bool RunCreateBranch(const FString& InBranchName, const FString& InComment, TArray<FString>& OutErrorMessages);
+
+/**
+ * Run branch rename <old name> <new name>
+ * @param	InOldName				The old name of the branch to rename
+ * @param	InNewName				The new name to rename the branch to
+ * @param	OutErrorMessages		Any errors (from StdErr) as an array per-line
+ */
+bool RunRenameBranch(const FString& InOldName, const FString& InNewName, TArray<FString>& OutErrorMessages);
+
+/**
+ * Run branch delete <name1> <name2 ...>
+ * @param	InBranchNames			The name of the branch(es) to delete
+ * @param	OutErrorMessages		Any errors (from StdErr) as an array per-line
  */
-void AddShelvedFileToChangelist(FPlasticSourceControlChangelistState& InOutChangelistsState, FString&& InFilename, EWorkspaceState InShelveStatus);
+bool RunDeleteBranches(const TArray<FString>& InBranchNames, TArray<FString>& OutErrorMessages);
 
 /**
  * Helper function for various commands to update cached states.
diff --git a/Engine/Plugins/Developer/PlasticSourceControl/Source/PlasticSourceControl/Private/PlasticSourceControlVersions.h b/Engine/Plugins/Developer/PlasticSourceControl/Source/PlasticSourceControl/Private/PlasticSourceControlVersions.h
--- a/Engine/Plugins/Developer/PlasticSourceControl/Source/PlasticSourceControl/Private/PlasticSourceControlVersions.h
+++ b/Engine/Plugins/Developer/PlasticSourceControl/Source/PlasticSourceControl/Private/PlasticSourceControlVersions.h
@@ -6,27 +6,44 @@
 #include "SoftwareVersion.h"
 
 /**
- * List of key versions of Plastic SCM enabling new features for "cm" CLI.
+ * List of key versions of Unity Version Control enabling new features for "cm" CLI.
  *
- * Plastic SCM version strings in the form "X.Y.Z.C", ie Major.Minor.Patch.Changeset (as returned by GetPlasticScmVersion)
+ * Unity Version Control version strings in the form "X.Y.Z.C", ie Major.Minor.Patch.Build (as returned by GetPlasticScmVersion)
 */
 namespace PlasticSourceControlVersions
 {
-	// Oldest supported version of "cm" (from 2021/01/05)
+	// Oldest supported version of "cm"
 	// 9.0.16.4839 cm changelist 'persistent' flag now contain a '--' prefix.
-	// https://www.plasticscm.com/download/releasenotes/9.0.16.4839
+	// https://plasticscm.com/download/releasenotes/9.0.16.4839 (2021/01/05)
 	static FSoftwareVersion OldestSupported(TEXT("9.0.16.4839"));
 
 	// 11.0.16.7248 add support for --descriptionfile for multi-line descriptions and support for special characters
-	// https://www.plasticscm.com/download/releasenotes/11.0.16.7248
+	// https://plasticscm.com/download/releasenotes/11.0.16.7248 (2022/07/28)
 	static const FSoftwareVersion NewChangelistFileArgs(TEXT("11.0.16.7248"));
 
 	// 11.0.16.7504 add support for 'cm shelveset apply' to only unshelved the desired changes from a shelve
-	// https://www.plasticscm.com/download/releasenotes/11.0.16.7504
+	// https://plasticscm.com/download/releasenotes/11.0.16.7504 (2022/10/13)
 	static const FSoftwareVersion ShelvesetApplySelection(TEXT("11.0.16.7504"));
 
 	// 11.0.16.7608 add support for history --limit. It displays the N last revisions of the specified items.
-	// https://www.plasticscm.com/download/releasenotes/11.0.16.7608
+	// https://plasticscm.com/download/releasenotes/11.0.16.7608 (2022/11/07)
+	// NOTE: this is also the first version aware of Unreal Engine 5 "UnrealEditor.exe"
 	static const FSoftwareVersion NewHistoryLimit(TEXT("11.0.16.7608"));
 
+	// 11.0.16.7665 add support for undocheckout --keepchanges. It allows undo checkout and preserve all local changes.
+	// https://plasticscm.com/download/releasenotes/11.0.16.7665 (2022/12/01)
+	static const FSoftwareVersion UndoCheckoutKeepChanges(TEXT("11.0.16.7665"));
+
+	// 11.0.16.7709 add support for status --iscochanged. It uses "CO+CH" vs "CO" to distinguish CheckedOut with/without changes.
+	// https://plasticscm.com/download/releasenotes/11.0.16.7709 (2023/01/12)
+	static const FSoftwareVersion StatusIsCheckedOutChanged(TEXT("11.0.16.7709"));
+
+	// 11.0.16.7726 add support for merge --xml
+	// https://plasticscm.com/download/releasenotes/11.0.16.7726 (2023/01/19)
+	static const FSoftwareVersion MergeXml(TEXT("11.0.16.7726"));
+
+	// 11.0.16.8101 Introducing Smart Locks
+	// https://plasticscm.com/download/releasenotes/11.0.16.8101 (2023/07/06)
+	static const FSoftwareVersion SmartLocks(TEXT("11.0.16.8101"));
+
 } // namespace PlasticSourceControlVersions
