PR #11295: [Source Control] Update Plastic SCM provider to 1.9.0 (from 1.8.0)

diff --git a/Engine/Plugins/Developer/PlasticSourceControl/Source/PlasticSourceControl/Private/SPlasticSourceControlSettings.cpp b/Engine/Plugins/Developer/PlasticSourceControl/Source/PlasticSourceControl/Private/SPlasticSourceControlSettings.cpp
--- a/Engine/Plugins/Developer/PlasticSourceControl/Source/PlasticSourceControl/Private/SPlasticSourceControlSettings.cpp
+++ b/Engine/Plugins/Developer/PlasticSourceControl/Source/PlasticSourceControl/Private/SPlasticSourceControlSettings.cpp
@@ -11,16 +11,13 @@
 #include "Misc/Paths.h"
 #include "Misc/App.h"
 #include "Misc/FileHelper.h"
-#include "Modules/ModuleManager.h"
-#include "Framework/Notifications/NotificationManager.h"
 #include "Styling/SlateTypes.h"
-#include "Widgets/Input/SEditableTextBox.h"
-#include "Widgets/Input/SMultiLineEditableTextBox.h"
 #include "Widgets/Input/SButton.h"
 #include "Widgets/Input/SCheckBox.h"
+#include "Widgets/Input/SEditableTextBox.h"
+#include "Widgets/Input/SMultiLineEditableTextBox.h"
 #include "Widgets/Layout/SBorder.h"
 #include "Widgets/Layout/SSeparator.h"
-#include "Widgets/Notifications/SNotificationList.h"
 #include "Widgets/SBoxPanel.h"
 #include "Widgets/Text/STextBlock.h"
 #include "Styling/AppStyle.h"
@@ -32,32 +29,33 @@ void SPlasticSourceControlSettings::Construct(const FArguments& InArgs)
 	const FSlateFontInfo Font = FAppStyle::GetFontStyle(TEXT("SourceControl.LoginWindow.Font"));
 
 	bAutoCreateIgnoreFile = CanAutoCreateIgnoreFile();
-	bAutoInitialCommit = true;
+	WorkspaceParams.bAutoInitialCommit = true;
+
+	WorkspaceParams.InitialCommitMessage = LOCTEXT("InitialCommitMessage", "Initial checkin");
+	WorkspaceParams.ServerUrl = FText::FromString(FPlasticSourceControlModule::Get().GetProvider().GetServerUrl());
 
-	InitialCommitMessage = LOCTEXT("InitialCommitMessage", "Initial checkin");
-	ServerUrl = FText::FromString(TEXT("YourOrganization@cloud"));
 	if (FApp::HasProjectName())
 	{
-		WorkspaceName = FText::FromString(FApp::GetProjectName());
-		RepositoryName = WorkspaceName;
+		WorkspaceParams.WorkspaceName = FText::FromString(FApp::GetProjectName());
+		WorkspaceParams.RepositoryName = WorkspaceParams.WorkspaceName;
 	}
 
 	ChildSlot
 	[
 		SNew(SVerticalBox)
-		// Versions (Plugin & Plastic SCM) useful eg to help diagnose issues from screenshots
+		// Versions (Plugin & Unity Version Control) useful eg to help diagnose issues from screenshots
 		+SVerticalBox::Slot()
 		.AutoHeight()
 		.Padding(2.0f)
 		.VAlign(VAlign_Center)
 		[
 			SNew(SHorizontalBox)
-			.ToolTipText(LOCTEXT("PlasticVersions_Tooltip", "Plastic SCM and Plugin versions"))
+			.ToolTipText(LOCTEXT("PlasticVersions_Tooltip", "Unity Version Control (formerly Plastic SCM) and Plugin versions"))
 			+SHorizontalBox::Slot()
 			.FillWidth(1.0f)
 			[
 				SNew(STextBlock)
-				.Text(LOCTEXT("PlasticVersions", "Plastic SCM version"))
+				.Text(LOCTEXT("PlasticVersions", "Unity Version Control"))
 				.Font(Font)
 			]
 			+SHorizontalBox::Slot()
@@ -68,30 +66,31 @@ void SPlasticSourceControlSettings::Construct(const FArguments& InArgs)
 				.Font(Font)
 			]
 		]
-		// Plastic SCM command line tool not available warning
+		// Unity Version Control command line tool not available warning
 		+SVerticalBox::Slot()
 		.FillHeight(1.0f)
 		.Padding(2.0f)
 		[
 			SNew(STextBlock)
 			.Visibility(this, &SPlasticSourceControlSettings::PlasticNotAvailable)
-			.ToolTipText(LOCTEXT("PlasticNotAvailable_Tooltip", "Failed to launch Plastic SCM 'cm' command line tool. You need to install it and make sure that 'cm' is on the Path and correctly configured."))
-			.Text(LOCTEXT("PlasticNotAvailable", "Plastic SCM Command Line tool 'cm' failed to start:"))
+			.ToolTipText(LOCTEXT("PlasticNotAvailable_Tooltip", "Failed to launch Unity Version Control 'cm' command line tool.\nYou need to install it and make sure it is correctly configured with your credentials."))
+			.Text(LOCTEXT("PlasticNotAvailable", "Unity Version Control Command Line tool 'cm' failed to start."))
 			.Font(Font)
+			.ColorAndOpacity(FLinearColor::Red)
 		]
-		// Path to the Plastic SCM binary
+		// Path to the Unity Version Control binary
 		+SVerticalBox::Slot()
 		.AutoHeight()
 		.Padding(2.0f)
 		.VAlign(VAlign_Center)
 		[
 			SNew(SHorizontalBox)
-			.ToolTipText(LOCTEXT("BinaryPathLabel_Tooltip", "Path to the Plastic SCM Command Line tool 'cm' binary"))
+			.ToolTipText(LOCTEXT("BinaryPathLabel_Tooltip", "Path to the Unity Version Control Command Line tool 'cm' executable"))
 			+SHorizontalBox::Slot()
 			.FillWidth(1.0f)
 			[
 				SNew(STextBlock)
-				.Text(LOCTEXT("PathLabel", "Plastic SCM path to cm"))
+				.Text(LOCTEXT("PathLabel", "Path to cm"))
 				.Font(Font)
 			]
 			+SHorizontalBox::Slot()
@@ -99,7 +98,7 @@ void SPlasticSourceControlSettings::Construct(const FArguments& InArgs)
 			[
 				SNew(SEditableTextBox)
 				.Text(this, &SPlasticSourceControlSettings::GetBinaryPathText)
-				.HintText(LOCTEXT("BinaryPathLabel", "Path to the Plastic SCM binary"))
+				.HintText(LOCTEXT("BinaryPathLabel", "Path to the Unity Version Control 'cm' executable"))
 				.OnTextCommitted(this, &SPlasticSourceControlSettings::OnBinaryPathTextCommited)
 				.Font(Font)
 			]
@@ -134,7 +133,7 @@ void SPlasticSourceControlSettings::Construct(const FArguments& InArgs)
 		.VAlign(VAlign_Center)
 		[
 			SNew(SHorizontalBox)
-			.ToolTipText(LOCTEXT("PlasticUserName_Tooltip", "User name configured for the Plastic SCM workspace"))
+			.ToolTipText(LOCTEXT("PlasticUserName_Tooltip", "User name configured for the Unity Version Control workspace"))
 			+SHorizontalBox::Slot()
 			.FillWidth(1.0f)
 			[
@@ -164,9 +163,9 @@ void SPlasticSourceControlSettings::Construct(const FArguments& InArgs)
 		.Padding(2.0f, 5.0f)
 		[
 			SNew(STextBlock)
-			.Visibility(this, &SPlasticSourceControlSettings::CanInitializePlasticWorkspace)
+			.Visibility(this, &SPlasticSourceControlSettings::CanCreatePlasticWorkspace)
 			.ToolTipText(LOCTEXT("WorkspaceNotFound_Tooltip", "No Workspace found at the level or above the current Project. Use the form to create a new one."))
-			.Text(LOCTEXT("WorkspaceNotFound", "Current Project is not in a Plastic SCM Workspace. Create a new one:"))
+			.Text(LOCTEXT("WorkspaceNotFound", "Current Project is not in a Unity Version Control Workspace. Create a new one:"))
 			.Font(Font)
 		]
 		// Workspace and Repository Name
@@ -176,7 +175,7 @@ void SPlasticSourceControlSettings::Construct(const FArguments& InArgs)
 		.VAlign(VAlign_Center)
 		[
 			SNew(SHorizontalBox)
-			.Visibility(this, &SPlasticSourceControlSettings::CanInitializePlasticWorkspace)
+			.Visibility(this, &SPlasticSourceControlSettings::CanCreatePlasticWorkspace)
 			+SHorizontalBox::Slot()
 			.FillWidth(1.0f)
 			[
@@ -215,7 +214,7 @@ void SPlasticSourceControlSettings::Construct(const FArguments& InArgs)
 		.VAlign(VAlign_Center)
 		[
 			SNew(SHorizontalBox)
-			.Visibility(this, &SPlasticSourceControlSettings::CanInitializePlasticWorkspace)
+			.Visibility(this, &SPlasticSourceControlSettings::CanCreatePlasticWorkspace)
 			.ToolTipText(LOCTEXT("ServerUrl_Tooltip", "Enter the Server URL in the form address:port (eg. YourOrganization@cloud, local, or something like ip:port, eg localhost:8087)"))
 			+SHorizontalBox::Slot()
 			.FillWidth(1.0f)
@@ -234,6 +233,23 @@ void SPlasticSourceControlSettings::Construct(const FArguments& InArgs)
 				.Font(Font)
 			]
 		]
+		// Option to create a Partial/Gluon Workspace designed
+		+SVerticalBox::Slot()
+		.AutoHeight()
+		.Padding(2.0f)
+		.VAlign(VAlign_Center)
+		[
+			SNew(SCheckBox)
+			.Visibility(this, &SPlasticSourceControlSettings::CanCreatePlasticWorkspace)
+			.ToolTipText(LOCTEXT("CreatePartialWorkspace_Tooltip", "Create the new workspace in Gluon/partial mode, designed for artists."))
+			.IsChecked(WorkspaceParams.bCreatePartialWorkspace)
+			.OnCheckStateChanged(this, &SPlasticSourceControlSettings::OnCheckedCreatePartialWorkspace)
+			[
+				SNew(STextBlock)
+				.Text(LOCTEXT("CreatePartialWorkspace", "Make the new workspace a Gluon partial workspace."))
+				.Font(Font)
+			]
+		]
 		// Option to add a 'ignore.conf' file at Workspace creation time
 		+SVerticalBox::Slot()
 		.AutoHeight()
@@ -241,7 +257,7 @@ void SPlasticSourceControlSettings::Construct(const FArguments& InArgs)
 		.VAlign(VAlign_Center)
 		[
 			SNew(SCheckBox)
-			.Visibility(this, &SPlasticSourceControlSettings::CanInitializePlasticWorkspace)
+			.Visibility(this, &SPlasticSourceControlSettings::CanCreatePlasticWorkspace)
 			.ToolTipText(LOCTEXT("CreateIgnoreFile_Tooltip", "Create and add a standard 'ignore.conf' file"))
 			.IsEnabled(this, &SPlasticSourceControlSettings::CanAutoCreateIgnoreFile)
 			.IsChecked(bAutoCreateIgnoreFile)
@@ -252,15 +268,15 @@ void SPlasticSourceControlSettings::Construct(const FArguments& InArgs)
 				.Font(Font)
 			]
 		]
-		// Option to Make the initial Plastic SCM checkin
+		// Option to Make the initial Unity Version Control checkin
 		+SVerticalBox::Slot()
 		.AutoHeight()
 		.Padding(2.0f)
 		.VAlign(VAlign_Center)
 		[
 			SNew(SHorizontalBox)
-			.Visibility(this, &SPlasticSourceControlSettings::CanInitializePlasticWorkspace)
-			.ToolTipText(LOCTEXT("InitialCommit_Tooltip", "Make the initial Plastic SCM checkin"))
+			.Visibility(this, &SPlasticSourceControlSettings::CanCreatePlasticWorkspace)
+			.ToolTipText(LOCTEXT("InitialCommit_Tooltip", "Make the initial Unity Version Control checkin"))
 			+SHorizontalBox::Slot()
 			.FillWidth(0.7f)
 			[
@@ -328,7 +344,7 @@ void SPlasticSourceControlSettings::Construct(const FArguments& InArgs)
 			.OnCheckStateChanged(this, &SPlasticSourceControlSettings::OnCheckedEnableVerboseLogs)
 			[
 				SNew(STextBlock)
-				.Text(LOCTEXT("EnableVerboseLogs", "Enable Source Control Verbose logs"))
+				.Text(LOCTEXT("EnableVerboseLogs", "Enable Revision Control Verbose logs"))
 				.Font(Font)
 			]
 		]
@@ -339,15 +355,15 @@ void SPlasticSourceControlSettings::Construct(const FArguments& InArgs)
 		.VAlign(VAlign_Center)
 		[
 			SNew(SHorizontalBox)
-			.Visibility(this, &SPlasticSourceControlSettings::CanInitializePlasticWorkspace)
+			.Visibility(this, &SPlasticSourceControlSettings::CanCreatePlasticWorkspace)
 			+SHorizontalBox::Slot()
 			.FillWidth(1.0f)
 			[
 				SNew(SButton)
-				.IsEnabled(this, &SPlasticSourceControlSettings::IsReadyToInitializePlasticWorkspace)
-				.Text(LOCTEXT("PlasticInitWorkspace", "Create a new Plastic SCM workspace for the current project"))
-				.ToolTipText(LOCTEXT("PlasticInitWorkspace_Tooltip", "Create and initialize a new Plastic SCM workspace and repository for the current project"))
-				.OnClicked(this, &SPlasticSourceControlSettings::OnClickedInitializePlasticWorkspace)
+				.IsEnabled(this, &SPlasticSourceControlSettings::IsReadyToCreatePlasticWorkspace)
+				.Text(LOCTEXT("PlasticInitWorkspace", "Create a new Unity Version Control workspace for the current project"))
+				.ToolTipText(LOCTEXT("PlasticInitWorkspace_Tooltip", "Create a new Unity Version Control repository and workspace and for the current project"))
+				.OnClicked(this, &SPlasticSourceControlSettings::OnClickedCreatePlasticWorkspace)
 				.HAlign(HAlign_Center)
 				.ContentPadding(6.0f)
 			]
@@ -390,7 +406,7 @@ void SPlasticSourceControlSettings::OnBinaryPathTextCommited(const FText& InText
 	const bool bChanged = Provider.AccessSettings().SetBinaryPath(InText.ToString());
 	if (bChanged)
 	{
-		// Re-Check provided Plastic binary path for each change
+		// Re-Check provided cm binary path for each change
 		Provider.CheckPlasticAvailability();
 		if (Provider.IsPlasticAvailable())
 		{
@@ -402,7 +418,7 @@ void SPlasticSourceControlSettings::OnBinaryPathTextCommited(const FText& InText
 FText SPlasticSourceControlSettings::GetVersions() const
 {
 	const FPlasticSourceControlProvider& Provider = FPlasticSourceControlModule::Get().GetProvider();
-	return FText::FromString(Provider.GetPlasticScmVersion().String + TEXT(" (plugin v") + Provider.GetPluginVersion() + TEXT(")"));
+	return FText::FromString(Provider.GetPlasticScmVersion().String + TEXT("\t(plugin v") + Provider.GetPluginVersion() + TEXT(")"));
 }
 
 FText SPlasticSourceControlSettings::GetPathToWorkspaceRoot() const
@@ -416,7 +432,7 @@ FText SPlasticSourceControlSettings::GetUserName() const
 }
 
 
-EVisibility SPlasticSourceControlSettings::CanInitializePlasticWorkspace() const
+EVisibility SPlasticSourceControlSettings::CanCreatePlasticWorkspace() const
 {
 	const FPlasticSourceControlProvider& Provider = FPlasticSourceControlModule::Get().GetProvider();
 	const bool bPlasticAvailable = Provider.IsPlasticAvailable();
@@ -424,43 +440,53 @@ EVisibility SPlasticSourceControlSettings::CanInitializePlasticWorkspace() const
 	return (bPlasticAvailable && !bPlasticWorkspaceFound) ? EVisibility::Visible : EVisibility::Collapsed;
 }
 
-bool SPlasticSourceControlSettings::IsReadyToInitializePlasticWorkspace() const
+bool SPlasticSourceControlSettings::IsReadyToCreatePlasticWorkspace() const
 {
 	// Workspace Name cannot be left empty
-	const bool bWorkspaceNameOk = !WorkspaceName.IsEmpty();
+	const bool bWorkspaceNameOk = !WorkspaceParams.WorkspaceName.IsEmpty();
 	// RepositoryName and ServerUrl should also be filled
-	const bool bRepositoryNameOk = !RepositoryName.IsEmpty() && !ServerUrl.IsEmpty();
+	const bool bRepositoryNameOk = !WorkspaceParams.RepositoryName.IsEmpty() && !WorkspaceParams.ServerUrl.IsEmpty();
 	// If Initial Commit is requested, checkin message cannot be empty
-	const bool bInitialCommitOk = (!bAutoInitialCommit || !InitialCommitMessage.IsEmpty());
+	const bool bInitialCommitOk = (!WorkspaceParams.bAutoInitialCommit || !WorkspaceParams.InitialCommitMessage.IsEmpty());
 	return bWorkspaceNameOk && bRepositoryNameOk && bInitialCommitOk;
 }
 
 
 void SPlasticSourceControlSettings::OnWorkspaceNameCommited(const FText& InText, ETextCommit::Type InCommitType)
 {
-	WorkspaceName = InText;
+	WorkspaceParams.WorkspaceName = InText;
 }
 FText SPlasticSourceControlSettings::GetWorkspaceName() const
 {
-	return WorkspaceName;
+	return WorkspaceParams.WorkspaceName;
 }
 
 void SPlasticSourceControlSettings::OnRepositoryNameCommited(const FText& InText, ETextCommit::Type InCommitType)
 {
-	RepositoryName = InText;
+	WorkspaceParams.RepositoryName = InText;
 }
 FText SPlasticSourceControlSettings::GetRepositoryName() const
 {
-	return RepositoryName;
+	return WorkspaceParams.RepositoryName;
 }
 
 void SPlasticSourceControlSettings::OnServerUrlCommited(const FText& InText, ETextCommit::Type InCommitType)
 {
-	ServerUrl = InText;
+	WorkspaceParams.ServerUrl = InText;
 }
 FText SPlasticSourceControlSettings::GetServerUrl() const
 {
-	return ServerUrl;
+	return WorkspaceParams.ServerUrl;
+}
+
+bool SPlasticSourceControlSettings::CreatePartialWorkspace() const
+{
+	return WorkspaceParams.bCreatePartialWorkspace;
+}
+
+void SPlasticSourceControlSettings::OnCheckedCreatePartialWorkspace(ECheckBoxState NewCheckedState)
+{
+	WorkspaceParams.bCreatePartialWorkspace = (NewCheckedState == ECheckBoxState::Checked);
 }
 
 bool SPlasticSourceControlSettings::CanAutoCreateIgnoreFile() const
@@ -476,181 +502,39 @@ void SPlasticSourceControlSettings::OnCheckedCreateIgnoreFile(ECheckBoxState New
 
 void SPlasticSourceControlSettings::OnCheckedInitialCommit(ECheckBoxState NewCheckedState)
 {
-	bAutoInitialCommit = (NewCheckedState == ECheckBoxState::Checked);
+	WorkspaceParams.bAutoInitialCommit = (NewCheckedState == ECheckBoxState::Checked);
 }
 
 void SPlasticSourceControlSettings::OnInitialCommitMessageCommited(const FText& InText, ETextCommit::Type InCommitType)
 {
-	InitialCommitMessage = InText;
+	WorkspaceParams.InitialCommitMessage = InText;
 }
 
 FText SPlasticSourceControlSettings::GetInitialCommitMessage() const
 {
-	return InitialCommitMessage;
+	return WorkspaceParams.InitialCommitMessage;
 }
 
 
-FReply SPlasticSourceControlSettings::OnClickedInitializePlasticWorkspace()
+FReply SPlasticSourceControlSettings::OnClickedCreatePlasticWorkspace()
 {
-	UE_LOG(LogSourceControl, Log, TEXT("InitializePlasticWorkspace(%s, %s, %s) CreateIgnore=%d Commit=%d"),
-		*WorkspaceName.ToString(), *RepositoryName.ToString(), *ServerUrl.ToString(), bAutoCreateIgnoreFile, bAutoInitialCommit);
-
-	// 1.a. Create a repository (if not already existing) and a workspace: launch an asynchronous MakeWorkspace operation
-	LaunchMakeWorkspaceOperation();
+	UE_LOG(LogSourceControl, Log, TEXT("CreatePlasticWorkspace(%s, %s, %s) PartialWorkspace=%d CreateIgnore=%d Commit=%d"),
+		*WorkspaceParams.WorkspaceName.ToString(), *WorkspaceParams.RepositoryName.ToString(), *WorkspaceParams.ServerUrl.ToString(),
+		WorkspaceParams.bCreatePartialWorkspace, bAutoCreateIgnoreFile, WorkspaceParams.bAutoInitialCommit);
 
-	return FReply::Handled();
-}
-
-
-/// 1. Create a repository (if not already existing) and a workspace
-void SPlasticSourceControlSettings::LaunchMakeWorkspaceOperation()
-{
-	TSharedRef<FPlasticMakeWorkspace, ESPMode::ThreadSafe> MakeWorkspaceOperation = ISourceControlOperation::Create<FPlasticMakeWorkspace>();
-	MakeWorkspaceOperation->WorkspaceName = WorkspaceName.ToString();
-	MakeWorkspaceOperation->RepositoryName = RepositoryName.ToString();
-	MakeWorkspaceOperation->ServerUrl = ServerUrl.ToString();
-
-	FPlasticSourceControlProvider& Provider = FPlasticSourceControlModule::Get().GetProvider();
-	ECommandResult::Type Result = Provider.Execute(MakeWorkspaceOperation, TArray<FString>(), EConcurrency::Asynchronous, FSourceControlOperationComplete::CreateSP(this, &SPlasticSourceControlSettings::OnSourceControlOperationComplete));
-	if (Result == ECommandResult::Succeeded)
+	if (bAutoCreateIgnoreFile)
 	{
-		DisplayInProgressNotification(MakeWorkspaceOperation->GetInProgressString());
+		// 1. Create a standard "ignore.conf" file with common patterns for a typical Blueprint & C++ project
+		CreateIgnoreFile();
 	}
-	else
-	{
-		DisplayFailureNotification(MakeWorkspaceOperation->GetName());
-	}
-}
-
-/// 2. Add all project files to Source Control (.uproject, Config/, Content/, Source/ files and ignore.conf if any)
-void SPlasticSourceControlSettings::LaunchMarkForAddOperation()
-{
-	TSharedRef<FMarkForAdd, ESPMode::ThreadSafe> MarkForAddOperation = ISourceControlOperation::Create<FMarkForAdd>();
-	FPlasticSourceControlProvider& Provider = FPlasticSourceControlModule::Get().GetProvider();
 
-	// 1.b. Check the new workspace status to enable connection
-	Provider.CheckPlasticAvailability();
+	// 2. Create a repository (if not already existing) and a workspace: launch an asynchronous MakeWorkspace operation
+	FPlasticSourceControlModule::Get().GetWorkspaceCreation().MakeWorkspace(WorkspaceParams);
 
-	if (Provider.IsWorkspaceFound())
-	{
-		if (bAutoCreateIgnoreFile)
-		{
-			// 1.c Create a standard "ignore.conf" file with common patterns for a typical Blueprint & C++ project
-			CreateIgnoreFile();
-		}
-		// 2. Add all project files to Source Control (.uproject, Config/, Content/, Source/ files and ignore.conf if any)
-		const TArray<FString> ProjectFiles = GetProjectFiles();
-		ECommandResult::Type Result = Provider.Execute(MarkForAddOperation, ProjectFiles, EConcurrency::Asynchronous, FSourceControlOperationComplete::CreateSP(this, &SPlasticSourceControlSettings::OnSourceControlOperationComplete));
-		if (Result == ECommandResult::Succeeded)
-		{
-			DisplayInProgressNotification(MarkForAddOperation->GetInProgressString());
-		}
-		else
-		{
-			DisplayFailureNotification(MarkForAddOperation->GetName());
-		}
-	}
-	else
-	{
-		DisplayFailureNotification(MarkForAddOperation->GetName());
-	}
-}
-
-/// 3. Launch an asynchronous "CheckIn" operation and start another ongoing notification
-void SPlasticSourceControlSettings::LaunchCheckInOperation()
-{
-	TSharedRef<FCheckIn, ESPMode::ThreadSafe> CheckInOperation = ISourceControlOperation::Create<FCheckIn>();
-	CheckInOperation->SetDescription(InitialCommitMessage);
-	FPlasticSourceControlProvider& Provider = FPlasticSourceControlModule::Get().GetProvider();
-	const TArray<FString> ProjectFiles = GetProjectFiles(); // Note: listing files and folders is only needed for the update status operation following the checkin to know on what to operate
-	ECommandResult::Type Result = Provider.Execute(CheckInOperation, ProjectFiles, EConcurrency::Asynchronous, FSourceControlOperationComplete::CreateSP(this, &SPlasticSourceControlSettings::OnSourceControlOperationComplete));
-	if (Result == ECommandResult::Succeeded)
-	{
-		DisplayInProgressNotification(CheckInOperation->GetInProgressString());
-	}
-	else
-	{
-		DisplayFailureNotification(CheckInOperation->GetName());
-	}
-}
-
-/// Delegate called when a source control operation has completed: launch the next one and manage notifications
-void SPlasticSourceControlSettings::OnSourceControlOperationComplete(const FSourceControlOperationRef& InOperation, ECommandResult::Type InResult)
-{
-	RemoveInProgressNotification();
-
-	// Report result with a notification
-	if (InResult == ECommandResult::Succeeded)
-	{
-		DisplaySuccessNotification(InOperation->GetName());
-	}
-	else
-	{
-		DisplayFailureNotification(InOperation->GetName());
-	}
-
-	// Launch the following asynchronous operation
-	if ((InOperation->GetName() == "MakeWorkspace") && (InResult == ECommandResult::Succeeded) && bAutoInitialCommit)
-	{
-		// 2. Add .uproject, Config/, Content/ and Source/ files (and ignore.conf if any)
-		LaunchMarkForAddOperation();
-	}
-	else if ((InOperation->GetName() == "MarkForAdd") && (InResult == ECommandResult::Succeeded) && bAutoInitialCommit)
-	{
-		// 3. optional initial Asynchronous commit with custom message: launch a "CheckIn" Operation
-		LaunchCheckInOperation();
-	}
-}
-
-// Display an ongoing notification during the whole operation
-void SPlasticSourceControlSettings::DisplayInProgressNotification(const FText& InOperationInProgressString)
-{
-	if (!OperationInProgressNotification.IsValid())
-	{
-		FNotificationInfo Info(InOperationInProgressString);
-		Info.bFireAndForget = false;
-		Info.ExpireDuration = 0.0f;
-		Info.FadeOutDuration = 1.0f;
-		OperationInProgressNotification = FSlateNotificationManager::Get().AddNotification(Info);
-		if (OperationInProgressNotification.IsValid())
-		{
-			OperationInProgressNotification.Pin()->SetCompletionState(SNotificationItem::CS_Pending);
-		}
-	}
-}
-
-// Remove the ongoing notification at the end of the operation
-void SPlasticSourceControlSettings::RemoveInProgressNotification()
-{
-	if (OperationInProgressNotification.IsValid())
-	{
-		OperationInProgressNotification.Pin()->ExpireAndFadeout();
-		OperationInProgressNotification.Reset();
-	}
-}
-
-// Display a temporary success notification at the end of the operation
-void SPlasticSourceControlSettings::DisplaySuccessNotification(const FName& InOperationName)
-{
-	const FText NotificationText = FText::Format(LOCTEXT("InitWorkspace_Success", "{0} operation was successful!"), FText::FromName(InOperationName));
-	FNotificationInfo Info(NotificationText);
-	Info.bUseSuccessFailIcons = true;
-	Info.Image = FAppStyle::GetBrush(TEXT("NotificationList.SuccessImage"));
-	FSlateNotificationManager::Get().AddNotification(Info);
-	UE_LOG(LogSourceControl, Verbose, TEXT("%s"), *NotificationText.ToString());
-}
-
-// Display a temporary failure notification at the end of the operation
-void SPlasticSourceControlSettings::DisplayFailureNotification(const FName& InOperationName)
-{
-	const FText NotificationText = FText::Format(LOCTEXT("InitWorkspace_Failure", "Error: {0} operation failed!"), FText::FromName(InOperationName));
-	FNotificationInfo Info(NotificationText);
-	Info.ExpireDuration = 8.0f;
-	FSlateNotificationManager::Get().AddNotification(Info);
-	UE_LOG(LogSourceControl, Error, TEXT("%s"), *NotificationText.ToString());
+	return FReply::Handled();
 }
 
-/** Delegate to check for presence of a Plastic ignore.conf file to an existing Plastic SCM workspace */
+/** Delegate to check for presence of an ignore.conf file to an existing Unity Version Control workspace */
 EVisibility SPlasticSourceControlSettings::CanAddIgnoreFile() const
 {
 	const bool bPlasticWorkspaceFound = FPlasticSourceControlModule::Get().GetProvider().IsWorkspaceFound();
@@ -658,12 +542,12 @@ EVisibility SPlasticSourceControlSettings::CanAddIgnoreFile() const
 	return (bPlasticWorkspaceFound && !bIgnoreFileFound) ? EVisibility::Visible : EVisibility::Collapsed;
 }
 
-/** Delegate to add a Plastic ignore.conf file to an existing Plastic SCM workspace */
+/** Delegate to add an ignore.conf file to an existing Unity Version Control workspace */
 FReply SPlasticSourceControlSettings::OnClickedAddIgnoreFile() const
 {
 	if (CreateIgnoreFile())
 	{
-		// Add ignore.conf to Plastic SCM
+		// Add ignore.conf to Unity Version Control
 		TArray<FString> InfoMessages;
 		TArray<FString> ErrorMessages;
 		TArray<FString> Parameters;
@@ -727,26 +611,8 @@ const FString SPlasticSourceControlSettings::GetIgnoreFileName() const
 /** Create a standard "ignore.conf" file with common patterns for a typical Blueprint & C++ project */
 bool SPlasticSourceControlSettings::CreateIgnoreFile() const
 {
-	const FString IgnoreFileContent = TEXT("Binaries\nBuild\nDerivedDataCache\nIntermediate\nSaved\nScript\nenc_temp_folder\n.idea\n.vscode\n.vs\n.vsconfig\n.ignore\n*.VC.db\n*.opensdf\n*.opendb\n*.sdf\n*.sln\n*.suo\n*.code-workspace\n*.xcodeproj\n*.xcworkspace");
+	const FString IgnoreFileContent = TEXT("Binaries\nBuild\nDerivedDataCache\nIntermediate\nSaved\nScript\nenc_temp_folder\n.idea\n.vscode\n.vs\n.ignore\n*.VC.db\n*.opensdf\n*.opendb\n*.sdf\n*.sln\n*.suo\n*.code-workspace\n*.xcodeproj\n*.xcworkspace");
 	return FFileHelper::SaveStringToFile(IgnoreFileContent, *GetIgnoreFileName(), FFileHelper::EEncodingOptions::ForceUTF8WithoutBOM);
 }
 
-/** List of files to add to Source Control (.uproject, Config/, Content/, Source/ files and ignore.conf if any) */
-TArray<FString> SPlasticSourceControlSettings::GetProjectFiles() const
-{
-	TArray<FString> ProjectFiles;
-	ProjectFiles.Add(FPaths::ConvertRelativePathToFull(FPaths::GetProjectFilePath()));
-	ProjectFiles.Add(FPaths::ConvertRelativePathToFull(FPaths::ProjectConfigDir()));
-	ProjectFiles.Add(FPaths::ConvertRelativePathToFull(FPaths::ProjectContentDir()));
-	if (FPaths::DirectoryExists(FPaths::GameSourceDir()))
-	{
-		ProjectFiles.Add(FPaths::ConvertRelativePathToFull(FPaths::GameSourceDir()));
-	}
-	if (FPaths::FileExists(GetIgnoreFileName()))
-	{
-		ProjectFiles.Add(GetIgnoreFileName());
-	}
-	return ProjectFiles;
-}
-
 #undef LOCTEXT_NAMESPACE
diff --git a/Engine/Plugins/Developer/PlasticSourceControl/Source/PlasticSourceControl/Private/SPlasticSourceControlSettings.h b/Engine/Plugins/Developer/PlasticSourceControl/Source/PlasticSourceControl/Private/SPlasticSourceControlSettings.h
--- a/Engine/Plugins/Developer/PlasticSourceControl/Source/PlasticSourceControl/Private/SPlasticSourceControlSettings.h
+++ b/Engine/Plugins/Developer/PlasticSourceControl/Source/PlasticSourceControl/Private/SPlasticSourceControlSettings.h
@@ -9,8 +9,7 @@
 #include "Input/Reply.h"
 #include "Styling/SlateTypes.h"
 
-#include "ISourceControlProvider.h"
-#include "ISourceControlOperation.h"
+#include "PlasticSourceControlWorkspaceCreation.h"
 
 class SPlasticSourceControlSettings : public SCompoundWidget
 {
@@ -36,18 +35,17 @@ private:
 	FText GetPathToWorkspaceRoot() const;
 	FText GetUserName() const;
 
-	/** Delegate to initialize a new Plastic workspace and repository */
-	EVisibility CanInitializePlasticWorkspace() const;
-	bool IsReadyToInitializePlasticWorkspace() const;
+	/** Delegate to create a new Plastic workspace and repository */
+	EVisibility CanCreatePlasticWorkspace() const;
+	bool IsReadyToCreatePlasticWorkspace() const;
 	void OnWorkspaceNameCommited(const FText& InText, ETextCommit::Type InCommitType);
 	FText GetWorkspaceName() const;
-	FText WorkspaceName;
 	void OnRepositoryNameCommited(const FText& InText, ETextCommit::Type InCommitType);
 	FText GetRepositoryName() const;
-	FText RepositoryName;
 	void OnServerUrlCommited(const FText& InText, ETextCommit::Type InCommitType);
 	FText GetServerUrl() const;
-	FText ServerUrl;
+	bool CreatePartialWorkspace() const;
+	void OnCheckedCreatePartialWorkspace(ECheckBoxState NewCheckedState);
 	bool CanAutoCreateIgnoreFile() const;
 	void OnCheckedCreateIgnoreFile(ECheckBoxState NewCheckedState);
 	bool bAutoCreateIgnoreFile;
@@ -62,28 +60,13 @@ private:
 	ECheckBoxState IsEnableVerboseLogsChecked() const;
 
 	void OnCheckedInitialCommit(ECheckBoxState NewCheckedState);
-	bool bAutoInitialCommit;
 	void OnInitialCommitMessageCommited(const FText& InText, ETextCommit::Type InCommitType);
 	FText GetInitialCommitMessage() const;
-	FText InitialCommitMessage;
 
-	/** Launch initial asynchronous add and commit operations */
-	void LaunchMakeWorkspaceOperation();
-	void LaunchMarkForAddOperation();
-	void LaunchCheckInOperation();
+	FReply OnClickedCreatePlasticWorkspace();
 
-	/** Delegate called when a source control operation has completed */
-	void OnSourceControlOperationComplete(const FSourceControlOperationRef& InOperation, ECommandResult::Type InResult);
-
-	/** Asynchronous operation progress notifications */
-	TWeakPtr<class SNotificationItem> OperationInProgressNotification;
-
-	void DisplayInProgressNotification(const FText& InOperationInProgressString);
-	void RemoveInProgressNotification();
-	void DisplaySuccessNotification(const FName& InOperationName);
-	void DisplayFailureNotification(const FName& InOperationName);
-
-	FReply OnClickedInitializePlasticWorkspace();
+	// Parameters for the workspace creation
+	FPlasticSourceControlWorkspaceCreation::FParameters WorkspaceParams;
 
 	/** Delegate to add a Plastic ignore.conf file to an existing Plastic workspace */
 	EVisibility CanAddIgnoreFile() const;
@@ -91,6 +74,4 @@ private:
 
 	const FString GetIgnoreFileName() const;
 	bool CreateIgnoreFile() const;
-
-	TArray<FString> GetProjectFiles() const;
 };
diff --git a/Engine/Plugins/Developer/PlasticSourceControl/Source/PlasticSourceControl/Private/ScopedTempFile.cpp b/Engine/Plugins/Developer/PlasticSourceControl/Source/PlasticSourceControl/Private/ScopedTempFile.cpp
--- a/Engine/Plugins/Developer/PlasticSourceControl/Source/PlasticSourceControl/Private/ScopedTempFile.cpp
+++ b/Engine/Plugins/Developer/PlasticSourceControl/Source/PlasticSourceControl/Private/ScopedTempFile.cpp
@@ -10,12 +10,12 @@
 
 FScopedTempFile::FScopedTempFile()
 {
-	Filename = FPaths::CreateTempFilename(*FPaths::ProjectLogDir(), TEXT("Plastic-Temp"), TEXT(".txt"));
+	Filename = FPaths::CreateTempFilename(*FPaths::ConvertRelativePathToFull(FPaths::ProjectLogDir()), TEXT("Plastic-Temp"), TEXT(".txt"));
 }
 
 FScopedTempFile::FScopedTempFile(const FString& InText)
 {
-	Filename = FPaths::CreateTempFilename(*FPaths::ProjectLogDir(), TEXT("Plastic-Temp"), TEXT(".txt"));
+	Filename = FPaths::CreateTempFilename(*FPaths::ConvertRelativePathToFull(FPaths::ProjectLogDir()), TEXT("Plastic-Temp"), TEXT(".txt"));
 	if (!FFileHelper::SaveStringToFile(InText, *Filename, FFileHelper::EEncodingOptions::ForceUTF8WithoutBOM))
 	{
 		UE_LOG(LogSourceControl, Error, TEXT("Failed to write to temp file: %s"), *Filename);
@@ -24,7 +24,7 @@ FScopedTempFile::FScopedTempFile(const FString& InText)
 
 FScopedTempFile::FScopedTempFile(const FText& InText)
 {
-	Filename = FPaths::CreateTempFilename(*FPaths::ProjectLogDir(), TEXT("Plastic-Temp"), TEXT(".txt"));
+	Filename = FPaths::CreateTempFilename(*FPaths::ConvertRelativePathToFull(FPaths::ProjectLogDir()), TEXT("Plastic-Temp"), TEXT(".txt"));
 	if (!FFileHelper::SaveStringToFile(InText.ToString(), *Filename, FFileHelper::EEncodingOptions::ForceUTF8WithoutBOM))
 	{
 		UE_LOG(LogSourceControl, Error, TEXT("Failed to write to temp file: %s"), *Filename);
diff --git a/Engine/Plugins/Developer/PlasticSourceControl/Source/PlasticSourceControl/Private/SoftwareVersion.cpp b/Engine/Plugins/Developer/PlasticSourceControl/Source/PlasticSourceControl/Private/SoftwareVersion.cpp
--- a/Engine/Plugins/Developer/PlasticSourceControl/Source/PlasticSourceControl/Private/SoftwareVersion.cpp
+++ b/Engine/Plugins/Developer/PlasticSourceControl/Source/PlasticSourceControl/Private/SoftwareVersion.cpp
@@ -5,14 +5,14 @@
 FSoftwareVersion::FSoftwareVersion(FString&& InVersionString)
 {
 	String = MoveTemp(InVersionString);
-	TArray<FString> Parts;
-	const int32 N = String.ParseIntoArray(Parts, TEXT("."));
-	if (N == 4)
+	TArray<FString> VersionElements;
+	String.ParseIntoArray(VersionElements, TEXT("."));
+	if (VersionElements.Num() == 4)
 	{
-		Major = FCString::Atoi(*Parts[0]);
-		Minor = FCString::Atoi(*Parts[1]);
-		Patch = FCString::Atoi(*Parts[2]);
-		Changeset = FCString::Atoi(*Parts[3]);
+		Major = FCString::Atoi(*VersionElements[0]);
+		Minor = FCString::Atoi(*VersionElements[1]);
+		Patch = FCString::Atoi(*VersionElements[2]);
+		Changeset = FCString::Atoi(*VersionElements[3]);
 	}
 }
 
diff --git a/Engine/Plugins/Developer/PlasticSourceControl/Source/PlasticSourceControl/Private/SoftwareVersion.h b/Engine/Plugins/Developer/PlasticSourceControl/Source/PlasticSourceControl/Private/SoftwareVersion.h
--- a/Engine/Plugins/Developer/PlasticSourceControl/Source/PlasticSourceControl/Private/SoftwareVersion.h
+++ b/Engine/Plugins/Developer/PlasticSourceControl/Source/PlasticSourceControl/Private/SoftwareVersion.h
@@ -15,10 +15,10 @@ struct FSoftwareVersion
 
 	FString String;
 
-	int Major = 0;
-	int Minor = 0;
-	int Patch = 0;
-	int Changeset = 0;
+	int32 Major = 0;
+	int32 Minor = 0;
+	int32 Patch = 0;
+	int32 Changeset = 0;
 };
 
 bool operator==(const FSoftwareVersion& Rhs, const FSoftwareVersion& Lhs);
