Fixed more bAllowShrinking usage.

diff --git a/Engine/Plugins/Developer/PlasticSourceControl/Source/PlasticSourceControl/Private/PlasticSourceControlParsers.cpp b/Engine/Plugins/Developer/PlasticSourceControl/Source/PlasticSourceControl/Private/PlasticSourceControlParsers.cpp
--- a/Engine/Plugins/Developer/PlasticSourceControl/Source/PlasticSourceControl/Private/PlasticSourceControlParsers.cpp
+++ b/Engine/Plugins/Developer/PlasticSourceControl/Source/PlasticSourceControl/Private/PlasticSourceControlParsers.cpp
@@ -1143,7 +1143,7 @@ bool ParseShelveDiffResult(const FString InWorkspaceRoot, TArray<FString>&& InRe
 		EWorkspaceState ShelveState = ParseShelveFileStatus(Result[0]);
 
 		// Remove outer double quotes
-		Result.MidInline(3, Result.Len() - 4, false);
+		Result.MidInline(3, Result.Len() - 4, EAllowShrinking::No);
 
 		FString MovedFrom;
 		if (ShelveState == EWorkspaceState::Moved)
@@ -1291,7 +1291,7 @@ bool ParseShelveDiffResults(const FString InWorkspaceRoot, TArray<FString>&& InR
 			const int32 BaseRevisionId = FCString::Atoi(*ResultElements[1]);
 			// Remove outer double quotes on filename
 			FString File = MoveTemp(ResultElements[2]);
-			File.MidInline(1, File.Len() - 2, false);
+			File.MidInline(1, File.Len() - 2, EAllowShrinking::No);
 			FString AbsoluteFilename = FPaths::ConvertRelativePathToFull(InWorkspaceRoot, File);
 
 			if (ShelveState == EWorkspaceState::Moved)
diff --git a/Engine/Plugins/Developer/PlasticSourceControl/Source/PlasticSourceControl/Private/PlasticSourceControlUtils.cpp b/Engine/Plugins/Developer/PlasticSourceControl/Source/PlasticSourceControl/Private/PlasticSourceControlUtils.cpp
--- a/Engine/Plugins/Developer/PlasticSourceControl/Source/PlasticSourceControl/Private/PlasticSourceControlUtils.cpp
+++ b/Engine/Plugins/Developer/PlasticSourceControl/Source/PlasticSourceControl/Private/PlasticSourceControlUtils.cpp
@@ -281,7 +281,7 @@ static bool RunStatus(const FString& InDir, TArray<FString>&& InFiles, const ESt
 		if (Results.Num() > 0)
 		{
 			PlasticSourceControlParsers::GetChangesetFromWorkspaceStatus(Results, OutChangeset);
-			Results.RemoveAt(0, 1, false);
+			Results.RemoveAt(0, 1, EAllowShrinking::No);
 		}
 
 		// Normalize file paths in the result (convert all '\' to '/')
diff --git a/Engine/Plugins/Runtime/nDisplay/Source/DisplayCluster/Public/Misc/DisplayClusterHelpers.h b/Engine/Plugins/Runtime/nDisplay/Source/DisplayCluster/Public/Misc/DisplayClusterHelpers.h
--- a/Engine/Plugins/Runtime/nDisplay/Source/DisplayCluster/Public/Misc/DisplayClusterHelpers.h
+++ b/Engine/Plugins/Runtime/nDisplay/Source/DisplayCluster/Public/Misc/DisplayClusterHelpers.h
@@ -180,7 +180,7 @@ namespace DisplayClusterHelpers
 
 			if (InPairSeparator.Len() > 0 && InData.Num() > 0)
 			{
-				ResultStr.RemoveAt(ResultStr.Len() - InPairSeparator.Len(), InPairSeparator.Len(), !bAddQuoutes);
+				ResultStr.RemoveAt(ResultStr.Len() - InPairSeparator.Len(), InPairSeparator.Len(), bAddQuoutes ? EAllowShrinking::No : EAllowShrinking::Yes);
 			}
 
 			if (bAddQuoutes)
