Buildfix: Fixed LOCTEXT error introduced in PR #31119614.

#rnx

diff --git a/Engine/Plugins/Developer/PlasticSourceControl/Source/PlasticSourceControl/Private/PlasticSourceControlState.cpp b/Engine/Plugins/Developer/PlasticSourceControl/Source/PlasticSourceControl/Private/PlasticSourceControlState.cpp
--- a/Engine/Plugins/Developer/PlasticSourceControl/Source/PlasticSourceControl/Private/PlasticSourceControlState.cpp
+++ b/Engine/Plugins/Developer/PlasticSourceControl/Source/PlasticSourceControl/Private/PlasticSourceControlState.cpp
@@ -247,9 +247,9 @@ FText FPlasticSourceControlState::GetDisplayTooltip() const
 	case EWorkspaceState::Controlled:
 		return FText();
 	case EWorkspaceState::CheckedOutChanged:
-		return LOCTEXT("CheckedOut_Tooltip", "Checked-out (changed)");
+		return LOCTEXT("CheckedOutChanged_Tooltip", "Checked-out (changed)");
 	case EWorkspaceState::CheckedOutUnchanged:
-		return LOCTEXT("CheckedOut_Tooltip", "Checked-out (unchanged)");
+		return LOCTEXT("CheckedOutUnchanged_Tooltip", "Checked-out (unchanged)");
 	case EWorkspaceState::Added:
 		return LOCTEXT("Added_Tooltip", "Added");
 	case EWorkspaceState::Moved:
