PR #12084: [Source Control] Update Plastic SCM provider to 1.11.0 (from 1.9.0)

Changeslist taken from GH as is with a few non-Unity compile fixes.

diff --git a/Engine/Plugins/Developer/PlasticSourceControl/PlasticSourceControl.uplugin b/Engine/Plugins/Developer/PlasticSourceControl/PlasticSourceControl.uplugin
--- a/Engine/Plugins/Developer/PlasticSourceControl/PlasticSourceControl.uplugin
+++ b/Engine/Plugins/Developer/PlasticSourceControl/PlasticSourceControl.uplugin
@@ -1,7 +1,7 @@
 {
 	"FileVersion": 3,
-	"Version": 90,
-	"VersionName": "1.9.0",
+	"Version": 110,
+	"VersionName": "1.11.0",
 	"FriendlyName": "Plastic SCM",
 	"Description": "Unity Version Control (formerly Plastic SCM)",
 	"Category": "Source Control",
diff --git a/Engine/Plugins/Developer/PlasticSourceControl/Source/PlasticSourceControl/PlasticSourceControl.Build.cs b/Engine/Plugins/Developer/PlasticSourceControl/Source/PlasticSourceControl/PlasticSourceControl.Build.cs
--- a/Engine/Plugins/Developer/PlasticSourceControl/Source/PlasticSourceControl/PlasticSourceControl.Build.cs
+++ b/Engine/Plugins/Developer/PlasticSourceControl/Source/PlasticSourceControl/PlasticSourceControl.Build.cs
@@ -11,12 +11,15 @@ public class PlasticSourceControl : ModuleRules
 			{
 				"Core",
 				"CoreUObject",
+				"Engine",
 				"Slate",
 				"SlateCore",
 				"InputCore",
 				"UnrealEd",
 				"LevelEditor",
+				"DesktopPlatform",
 				"SourceControl",
+				"SourceControlWindows",
 				"XmlParser",
 				"Projects",
 				"AssetRegistry",
diff --git a/Engine/Plugins/Developer/PlasticSourceControl/Source/PlasticSourceControl/Private/Notification.cpp b/Engine/Plugins/Developer/PlasticSourceControl/Source/PlasticSourceControl/Private/Notification.cpp
--- a/Engine/Plugins/Developer/PlasticSourceControl/Source/PlasticSourceControl/Private/Notification.cpp
+++ b/Engine/Plugins/Developer/PlasticSourceControl/Source/PlasticSourceControl/Private/Notification.cpp
@@ -84,18 +84,13 @@ void FNotification::DisplaySuccess(const FSourceControlOperationBase& InOperatio
 	}
 	else
 	{
-		DisplaySuccess(InOperation.GetName());
-	}
-}
-
-void FNotification::DisplaySuccess(const FName& InOperationName)
-{
-	const FText NotificationText = FText::Format(
-		LOCTEXT("PlasticSourceControlOperation_Success", "{0} operation was successful."),
-		FText::FromName(InOperationName)
-	);
+		const FText NotificationText = FText::Format(
+			LOCTEXT("PlasticSourceControlOperation_Success", "{0} operation was successful."),
+			FText::FromName(InOperation.GetName())
+		);
 
-	DisplaySuccess(NotificationText);
+		DisplaySuccess(NotificationText);
+	}
 }
 
 void FNotification::DisplaySuccess(const FText& InNotificationText)
@@ -115,22 +110,24 @@ void FNotification::DisplayFailure(const FSourceControlOperationBase& InOperatio
 {
 	if (InOperation.GetResultInfo().ErrorMessages.Num() > 0)
 	{
-		DisplayFailure(InOperation.GetResultInfo().ErrorMessages[0]);
+		// If there are multiple messages, display the last one to not let the user with a notification starting with a "wait" or "in progress" message
+		const FText NotificationText = FText::Format(
+			LOCTEXT("PlasticSourceControlOperation_Failure", "Error: {0} operation failed!\n{1}"),
+			FText::FromName(InOperation.GetName()),
+			InOperation.GetResultInfo().ErrorMessages.Last()
+		);
+
+		DisplayFailure(NotificationText);
 	}
 	else
 	{
-		DisplayFailure(InOperation.GetName());
-	}
-}
+		const FText NotificationText = FText::Format(
+			LOCTEXT("PlasticSourceControlOperation_Failure", "Error: {0} operation failed!"),
+			FText::FromName(InOperation.GetName())
+		);
 
-void FNotification::DisplayFailure(const FName& InOperationName)
-{
-	const FText NotificationText = FText::Format(
-		LOCTEXT("PlasticSourceControlOperation_Failure", "Error: {0} operation failed!"),
-		FText::FromName(InOperationName)
-	);
-
-	DisplayFailure(NotificationText);
+		DisplayFailure(NotificationText);
+	}
 }
 
 void FNotification::DisplayFailure(const FText& InNotificationText)
diff --git a/Engine/Plugins/Developer/PlasticSourceControl/Source/PlasticSourceControl/Private/Notification.h b/Engine/Plugins/Developer/PlasticSourceControl/Source/PlasticSourceControl/Private/Notification.h
--- a/Engine/Plugins/Developer/PlasticSourceControl/Source/PlasticSourceControl/Private/Notification.h
+++ b/Engine/Plugins/Developer/PlasticSourceControl/Source/PlasticSourceControl/Private/Notification.h
@@ -28,11 +28,9 @@ public:
 	static void DisplayResult(const FSourceControlOperationBase& InOperation, ECommandResult::Type InResult);
 
 	static void DisplaySuccess(const FSourceControlOperationBase& InOperation);
-	static void DisplaySuccess(const FName& InOperationName);
 	static void DisplaySuccess(const FText& InNotificationText);
 
 	static void DisplayFailure(const FSourceControlOperationBase& InOperation);
-	static void DisplayFailure(const FName& InOperationName);
 	static void DisplayFailure(const FText& InNotificationText);
 
 private:
diff --git a/Engine/Plugins/Developer/PlasticSourceControl/Source/PlasticSourceControl/Private/PlasticSourceControlBranch.h b/Engine/Plugins/Developer/PlasticSourceControl/Source/PlasticSourceControl/Private/PlasticSourceControlBranch.h
--- a/Engine/Plugins/Developer/PlasticSourceControl/Source/PlasticSourceControl/Private/PlasticSourceControlBranch.h
+++ b/Engine/Plugins/Developer/PlasticSourceControl/Source/PlasticSourceControl/Private/PlasticSourceControlBranch.h
@@ -7,14 +7,6 @@
 class FPlasticSourceControlBranch
 {
 public:
-	FPlasticSourceControlBranch() = default;
-	FPlasticSourceControlBranch(const FString& InName, const FString& InRepository, const FString& InCreatedBy, const FDateTime& InDate, const FString& InComment)
-		: Name(InName)
-		, Repository(InRepository)
-		, CreatedBy(InCreatedBy)
-		, Date(InDate)
-		, Comment(InComment)
-	{}
 	FString Name;
 	FString Repository;
 	FString CreatedBy;
@@ -30,4 +22,4 @@ public:
 };
 
 typedef TSharedRef<class FPlasticSourceControlBranch, ESPMode::ThreadSafe> FPlasticSourceControlBranchRef;
-typedef TSharedPtr<class FPlasticSourceControlBranch, ESPMode::ThreadSafe> FPlasticSourceControlBranchPtr;
\ No newline at end of file
+typedef TSharedPtr<class FPlasticSourceControlBranch, ESPMode::ThreadSafe> FPlasticSourceControlBranchPtr;
diff --git a/Engine/Plugins/Developer/PlasticSourceControl/Source/PlasticSourceControl/Private/PlasticSourceControlBranchesWindow.cpp b/Engine/Plugins/Developer/PlasticSourceControl/Source/PlasticSourceControl/Private/PlasticSourceControlBranchesWindow.cpp
--- a/Engine/Plugins/Developer/PlasticSourceControl/Source/PlasticSourceControl/Private/PlasticSourceControlBranchesWindow.cpp
+++ b/Engine/Plugins/Developer/PlasticSourceControl/Source/PlasticSourceControl/Private/PlasticSourceControlBranchesWindow.cpp
@@ -2,31 +2,27 @@
 
 #include "PlasticSourceControlBranchesWindow.h"
 
-#include "Widgets/Docking/SDockTab.h"
-
-#include "PlasticSourceControlStyle.h"
 #include "SPlasticSourceControlBranchesWidget.h"
 
-#define LOCTEXT_NAMESPACE "PlasticSourceControlWindow"
+#include "Framework/Docking/TabManager.h"
+#include "RevisionControlStyle/RevisionControlStyle.h"
+#include "Widgets/Docking/SDockTab.h"
+
+#define LOCTEXT_NAMESPACE "PlasticSourceControlBranchesWindow"
 
-static const FName PlasticSourceControlWindowTabName("PlasticSourceControlWindow");
+static const FName PlasticSourceControlBranchesWindowTabName("PlasticSourceControlBranchesWindow");
 
 void FPlasticSourceControlBranchesWindow::Register()
 {
-	FPlasticSourceControlStyle::Initialize();
-	FPlasticSourceControlStyle::ReloadTextures();
-
-	FGlobalTabmanager::Get()->RegisterNomadTabSpawner(PlasticSourceControlWindowTabName, FOnSpawnTab::CreateRaw(this, &FPlasticSourceControlBranchesWindow::OnSpawnTab))
-		.SetDisplayName(LOCTEXT("PlasticSourceControlWindowTabTitle", "View Branches"))
+	FGlobalTabmanager::Get()->RegisterNomadTabSpawner(PlasticSourceControlBranchesWindowTabName, FOnSpawnTab::CreateRaw(this, &FPlasticSourceControlBranchesWindow::OnSpawnTab))
+		.SetDisplayName(LOCTEXT("PlasticSourceControlBranchesWindowTabTitle", "View Branches"))
 		.SetMenuType(ETabSpawnerMenuType::Hidden)
-		.SetIcon(FSlateIcon(FPlasticSourceControlStyle::Get().GetStyleSetName(), "PlasticSourceControl.PluginIcon.Small"));
+		.SetIcon(FSlateIcon(FRevisionControlStyleManager::GetStyleSetName(), "RevisionControl.Branched"));
 }
 
 void FPlasticSourceControlBranchesWindow::Unregister()
 {
-	FGlobalTabmanager::Get()->UnregisterNomadTabSpawner(PlasticSourceControlWindowTabName);
-
-	FPlasticSourceControlStyle::Shutdown();
+	FGlobalTabmanager::Get()->UnregisterNomadTabSpawner(PlasticSourceControlBranchesWindowTabName);
 }
 
 TSharedRef<SDockTab> FPlasticSourceControlBranchesWindow::OnSpawnTab(const FSpawnTabArgs& SpawnTabArgs)
@@ -40,7 +36,7 @@ TSharedRef<SDockTab> FPlasticSourceControlBranchesWindow::OnSpawnTab(const FSpaw
 
 void FPlasticSourceControlBranchesWindow::OpenTab()
 {
-	FGlobalTabmanager::Get()->TryInvokeTab(PlasticSourceControlWindowTabName);
+	FGlobalTabmanager::Get()->TryInvokeTab(PlasticSourceControlBranchesWindowTabName);
 }
 
 TSharedPtr<SWidget> FPlasticSourceControlBranchesWindow::CreateBranchesWidget()
@@ -48,4 +44,4 @@ TSharedPtr<SWidget> FPlasticSourceControlBranchesWindow::CreateBranchesWidget()
 	return SNew(SPlasticSourceControlBranchesWidget);
 }
 
-#undef LOCTEXT_NAMESPACE
\ No newline at end of file
+#undef LOCTEXT_NAMESPACE
diff --git a/Engine/Plugins/Developer/PlasticSourceControl/Source/PlasticSourceControl/Private/PlasticSourceControlChangelistState.h b/Engine/Plugins/Developer/PlasticSourceControl/Source/PlasticSourceControl/Private/PlasticSourceControlChangelistState.h
--- a/Engine/Plugins/Developer/PlasticSourceControl/Source/PlasticSourceControl/Private/PlasticSourceControlChangelistState.h
+++ b/Engine/Plugins/Developer/PlasticSourceControl/Source/PlasticSourceControl/Private/PlasticSourceControlChangelistState.h
@@ -5,6 +5,7 @@
 #include "CoreMinimal.h"
 
 #include "ISourceControlChangelistState.h"
+#include "ISourceControlState.h"
 
 #include "PlasticSourceControlChangelist.h"
 
diff --git a/Engine/Plugins/Developer/PlasticSourceControl/Source/PlasticSourceControl/Private/PlasticSourceControlCommand.h b/Engine/Plugins/Developer/PlasticSourceControl/Source/PlasticSourceControl/Private/PlasticSourceControlCommand.h
--- a/Engine/Plugins/Developer/PlasticSourceControl/Source/PlasticSourceControl/Private/PlasticSourceControlCommand.h
+++ b/Engine/Plugins/Developer/PlasticSourceControl/Source/PlasticSourceControl/Private/PlasticSourceControlCommand.h
@@ -34,7 +34,7 @@ public:
 	/**
 	 * This method is also used to tell the object to cleanup but not before
 	 * the object has finished it's work.
-	 */ 
+	 */
 	virtual void DoThreadedWork() override;
 
 	/** Save any results and call any registered callbacks. */
@@ -74,6 +74,9 @@ public:
 	/** Name of the current branch */
 	FString BranchName;
 
+	/** Name of the object that the workspace is switched to: usually the branch, sometimes a changeset or a label */
+	FString WorkspaceSelector;
+
 	/** Current Changeset Number */
 	int32 ChangesetNumber;
 
diff --git a/Engine/Plugins/Developer/PlasticSourceControl/Source/PlasticSourceControl/Private/PlasticSourceControlMenu.cpp b/Engine/Plugins/Developer/PlasticSourceControl/Source/PlasticSourceControl/Private/PlasticSourceControlMenu.cpp
--- a/Engine/Plugins/Developer/PlasticSourceControl/Source/PlasticSourceControl/Private/PlasticSourceControlMenu.cpp
+++ b/Engine/Plugins/Developer/PlasticSourceControl/Source/PlasticSourceControl/Private/PlasticSourceControlMenu.cpp
@@ -2,10 +2,12 @@
 
 #include "PlasticSourceControlMenu.h"
 
-#include "PlasticSourceControlBranchesWindow.h"
+#include "PlasticSourceControlLock.h"
 #include "PlasticSourceControlModule.h"
-#include "PlasticSourceControlProvider.h"
 #include "PlasticSourceControlOperations.h"
+#include "PlasticSourceControlProvider.h"
+#include "PlasticSourceControlUtils.h"
+#include "PlasticSourceControlVersions.h"
 #include "SPlasticSourceControlStatusBar.h"
 
 #include "ISourceControlModule.h"
@@ -14,9 +16,10 @@
 
 #include "ContentBrowserMenuContexts.h"
 #include "Interfaces/IPluginManager.h"
-#include "Modules/ModuleManager.h"
+#include "HAL/PlatformProcess.h"
 #include "LevelEditor.h"
 #include "Misc/MessageDialog.h"
+#include "Modules/ModuleManager.h"
 #include "Styling/AppStyle.h"
 
 #include "PackageUtils.h"
@@ -54,6 +57,7 @@ void FPlasticSourceControlMenu::Unregister()
 		return;
 	}
 
+	// Unregister the menu extensions from the level editor
 	if (UToolMenus* ToolMenus = UToolMenus::TryGet())
 	{
 		ToolMenus->UnregisterOwnerByName(UnityVersionControlMainMenuOwnerName);
@@ -93,6 +97,8 @@ void FPlasticSourceControlMenu::ExtendRevisionControlMenu()
 		if (FToolMenuSection* Section = ToolsMenu->FindSection("Source Control"))
 		{
 			AddViewBranches(*Section);
+			AddViewChangesets(*Section);
+			AddViewLocks(*Section);
 		}
 	}
 }
@@ -127,16 +133,20 @@ void FPlasticSourceControlMenu::ExtendAssetContextMenu()
 
 void FPlasticSourceControlMenu::GeneratePlasticAssetContextMenu(FMenuBuilder& MenuBuilder, TArray<FAssetData> InAssetObjectPaths)
 {
+	const FPlasticSourceControlProvider& Provider = FPlasticSourceControlModule::Get().GetProvider();
+	const TArray<FString> Files = PackageUtils::AssetDataToFileNames(InAssetObjectPaths);
+	const TArray<FPlasticSourceControlLockRef> SelectedLocks = PlasticSourceControlUtils::GetLocksForWorkingBranch(Provider, Files);
+
 	MenuBuilder.BeginSection("AssetPlasticActions", LOCTEXT("UnityVersionControlAssetContextLocksMenuHeading", "Unity Version Control Locks"));
 
 	{
 		MenuBuilder.AddMenuEntry(
 			LOCTEXT("PlasticReleaseLock", "Release Lock"),
-			LOCTEXT("PlasticReleaseLockTooltip", "Release Lock(s) on the selected assets. Requires administrator privileges on the server."),
+			LOCTEXT("PlasticReleaseLockTooltip", "Release Lock(s) on the selected assets.\nReleasing locks will allow other users to keep working on these files and retrieve locks (on the same branch, in the latest revision)."),
 			FSlateIcon(FAppStyle::GetAppStyleSetName(), "PropertyWindow.Unlocked"),
 			FUIAction(
-				FExecuteAction::CreateRaw(this, &FPlasticSourceControlMenu::ExecuteReleaseLocks, InAssetObjectPaths),
-				FCanExecuteAction::CreateRaw(this, &FPlasticSourceControlMenu::CanReleaseLocks, InAssetObjectPaths)
+				FExecuteAction::CreateRaw(this, &FPlasticSourceControlMenu::ExecuteReleaseLocks, SelectedLocks),
+				FCanExecuteAction::CreateRaw(this, &FPlasticSourceControlMenu::CanReleaseLocks, SelectedLocks)
 			)
 		);
 	}
@@ -144,11 +154,11 @@ void FPlasticSourceControlMenu::GeneratePlasticAssetContextMenu(FMenuBuilder& Me
 	{
 		MenuBuilder.AddMenuEntry(
 			LOCTEXT("PlasticRemoveLock", "Remove Lock"),
-			LOCTEXT("PlasticRemoveLockTooltip", "Remove/Delete Lock(s) on the selected assets. Requires administrator privileges on the server."),
+			LOCTEXT("PlasticRemoveLockTooltip", "Remove Lock(s) on the selected assets.\nRemoving locks will allow other users to edit these files anywhere (on any branch) increasing the risk of future merge conflicts."),
 			FSlateIcon(FAppStyle::GetAppStyleSetName(), "PropertyWindow.Unlocked"),
 			FUIAction(
-				FExecuteAction::CreateRaw(this, &FPlasticSourceControlMenu::ExecuteRemoveLocks, InAssetObjectPaths),
-				FCanExecuteAction::CreateRaw(this, &FPlasticSourceControlMenu::CanRemoveLocks, InAssetObjectPaths)
+				FExecuteAction::CreateRaw(this, &FPlasticSourceControlMenu::ExecuteRemoveLocks, SelectedLocks),
+				FCanExecuteAction::CreateRaw(this, &FPlasticSourceControlMenu::CanRemoveLocks, SelectedLocks)
 			)
 		);
 	}
@@ -170,16 +180,12 @@ void FPlasticSourceControlMenu::GeneratePlasticAssetContextMenu(FMenuBuilder& Me
 	MenuBuilder.EndSection();
 }
 
-bool FPlasticSourceControlMenu::CanReleaseLocks(TArray<FAssetData> InAssetObjectPaths) const
+bool FPlasticSourceControlMenu::CanReleaseLocks(TArray<FPlasticSourceControlLockRef> InSelectedLocks) const
 {
-	const TArray<FString> Files = PackageUtils::AssetDataToFileNames(InAssetObjectPaths);
-
-	for (const FString& File : Files)
+	for (const FPlasticSourceControlLockRef& Lock : InSelectedLocks)
 	{
-		const FString AbsoluteFilename = FPaths::ConvertRelativePathToFull(File);
-		const auto State = FPlasticSourceControlModule::Get().GetProvider().GetStateInternal(AbsoluteFilename);
-		// If exclusively Checked Out (Locked) the lock can be released coming back to it's potential underlying "Retained" status if changes where already checked in the branch
-		if (!State->LockedBy.IsEmpty() && State->LockedId != ISourceControlState::INVALID_REVISION)
+		// If "Locked" (currently exclusively Checked Out) the lock can be Released, coming back to it's potential underlying "Retained" status if changes where already checked in the branch
+		if (Lock->bIsLocked)
 		{
 			return true;
 		}
@@ -188,45 +194,34 @@ bool FPlasticSourceControlMenu::CanReleaseLocks(TArray<FAssetData> InAssetObject
 	return false;
 }
 
-bool FPlasticSourceControlMenu::CanRemoveLocks(TArray<FAssetData> InAssetObjectPaths) const
+bool FPlasticSourceControlMenu::CanRemoveLocks(TArray<FPlasticSourceControlLockRef> InSelectedLocks) const
 {
-	const TArray<FString> Files = PackageUtils::AssetDataToFileNames(InAssetObjectPaths);
-
-	for (const FString& File : Files)
-	{
-		const FString AbsoluteFilename = FPaths::ConvertRelativePathToFull(File);
-		const auto State = FPlasticSourceControlModule::Get().GetProvider().GetStateInternal(AbsoluteFilename);
-		// If Locked or Retained, the lock can be removed, that is completely deleted in order to simply ignore the changes from the branch
-		if (State->LockedId != ISourceControlState::INVALID_REVISION)
-		{
-			return true;
-		}
-	}
-
-	return false;
+	// All "Locked" or "Retained" locks can be Removed
+	return (InSelectedLocks.Num() > 0);
 }
 
-void FPlasticSourceControlMenu::ExecuteReleaseLocks(TArray<FAssetData> InAssetObjectPaths)
+void FPlasticSourceControlMenu::ExecuteReleaseLocks(TArray<FPlasticSourceControlLockRef> InSelectedLocks)
 {
-	ExecuteUnlock(InAssetObjectPaths, false);
+	ExecuteUnlock(MoveTemp(InSelectedLocks), false);
 }
 
-void FPlasticSourceControlMenu::ExecuteRemoveLocks(TArray<FAssetData> InAssetObjectPaths)
+void FPlasticSourceControlMenu::ExecuteRemoveLocks(TArray<FPlasticSourceControlLockRef> InSelectedLocks)
 {
-	ExecuteUnlock(InAssetObjectPaths, true);
+	ExecuteUnlock(MoveTemp(InSelectedLocks), true);
 }
 
-void FPlasticSourceControlMenu::ExecuteUnlock(const TArray<FAssetData>& InAssetObjectPaths, const bool bInRemove)
+void FPlasticSourceControlMenu::ExecuteUnlock(TArray<FPlasticSourceControlLockRef>&& InSelectedLocks, const bool bInRemove)
 {
 	if (!Notification.IsInProgress())
 	{
-		const TArray<FString> Files = PackageUtils::AssetDataToFileNames(InAssetObjectPaths);
-
 		// Launch a custom "Release/Remove Lock" operation
 		FPlasticSourceControlProvider& Provider = FPlasticSourceControlModule::Get().GetProvider();
+		const FString& WorkspaceRoot = Provider.GetPathToWorkspaceRoot();
+		const TArray<FString> Files = PlasticSourceControlUtils::LocksToFileNames(WorkspaceRoot, InSelectedLocks);
 		TSharedRef<FPlasticUnlock, ESPMode::ThreadSafe> UnlockOperation = ISourceControlOperation::Create<FPlasticUnlock>();
-		const ECommandResult::Type Result = Provider.Execute(UnlockOperation, Files, EConcurrency::Asynchronous, FSourceControlOperationComplete::CreateRaw(this, &FPlasticSourceControlMenu::OnSourceControlOperationComplete));
 		UnlockOperation->bRemove = bInRemove;
+		UnlockOperation->Locks = MoveTemp(InSelectedLocks);
+		const ECommandResult::Type Result = Provider.Execute(UnlockOperation, Files, EConcurrency::Asynchronous, FSourceControlOperationComplete::CreateRaw(this, &FPlasticSourceControlMenu::OnSourceControlOperationComplete));
 		if (Result == ECommandResult::Succeeded)
 		{
 			// Display an ongoing notification during the whole operation (packages will be reloaded at the completion of the operation)
@@ -445,11 +440,12 @@ void FPlasticSourceControlMenu::VisitSupportURLClicked() const
 
 void FPlasticSourceControlMenu::VisitLockRulesURLClicked(const FString InOrganizationName) const
 {
-	const FString OrganizationLockRulesURL = FString::Printf(
-		TEXT("https://dashboard.unity3d.com/devops/organizations/default/plastic-scm/organizations/%s/lock-rules"),
-		*InOrganizationName
-	);
-	FPlatformProcess::LaunchURL(*OrganizationLockRulesURL, NULL, NULL);
+	PlasticSourceControlUtils::OpenLockRulesInCloudDashboard(InOrganizationName);
+}
+
+void FPlasticSourceControlMenu::OpenDesktopApplication() const
+{
+	PlasticSourceControlUtils::OpenDesktopApplication();
 }
 
 void FPlasticSourceControlMenu::OpenBranchesWindow() const
@@ -457,6 +453,16 @@ void FPlasticSourceControlMenu::OpenBranchesWindow() const
 	FPlasticSourceControlModule::Get().GetBranchesWindow().OpenTab();
 }
 
+void FPlasticSourceControlMenu::OpenChangesetsWindow() const
+{
+	FPlasticSourceControlModule::Get().GetChangesetsWindow().OpenTab();
+}
+
+void FPlasticSourceControlMenu::OpenLocksWindow() const
+{
+	FPlasticSourceControlModule::Get().GetLocksWindow().OpenTab();
+}
+
 void FPlasticSourceControlMenu::OnSyncAllOperationComplete(const FSourceControlOperationRef& InOperation, ECommandResult::Type InResult)
 {
 	OnSourceControlOperationComplete(InOperation, InResult);
@@ -482,8 +488,11 @@ void FPlasticSourceControlMenu::OnSourceControlOperationComplete(const FSourceCo
 	FNotification::DisplayResult(InOperation, InResult);
 }
 
+// TODO rework the menus with sub-menus like in the POC branch
 void FPlasticSourceControlMenu::AddMenuExtension(FToolMenuSection& Menu)
 {
+	FPlasticSourceControlProvider& Provider = FPlasticSourceControlModule::Get().GetProvider();
+
 	Menu.AddMenuEntry(
 		"PlasticSync",
 		LOCTEXT("PlasticSync",			"Sync/Update Workspace"),
@@ -599,7 +608,23 @@ void FPlasticSourceControlMenu::AddMenuExtension(FToolMenuSection& Menu)
 		);
 	}
 
+	Menu.AddMenuEntry(
+		"PlasticDesktopApp",
+		TAttribute<FText>::CreateLambda([&Provider]()
+			{
+				return Provider.IsPartialWorkspace() ? LOCTEXT("PlasticGluon", "Open in Gluon") : LOCTEXT("PlasticDesktopApp", "Open in Desktop App");
+			}),
+		TAttribute<FText>::CreateLambda([&Provider]()
+			{
+				return Provider.IsPartialWorkspace() ? LOCTEXT("PlasticDesktopAppTooltip", "Open the workspace in Unity Version Control Gluon Application.") : LOCTEXT("PlasticGluonTooltip", "Open the workspace in Unity Version Control Desktop Application.");
+			}),
+		FSlateIcon(FAppStyle::GetAppStyleSetName(), "ClassIcon.UserDefinedStruct"),
+		FUIAction(FExecuteAction::CreateRaw(this, &FPlasticSourceControlMenu::OpenDesktopApplication))
+	);
+
 	AddViewBranches(Menu);
+	AddViewChangesets(Menu);
+	AddViewLocks(Menu);
 }
 
 void FPlasticSourceControlMenu::AddViewBranches(FToolMenuSection& Menu)
@@ -616,4 +641,34 @@ void FPlasticSourceControlMenu::AddViewBranches(FToolMenuSection& Menu)
 	);
 }
 
+void FPlasticSourceControlMenu::AddViewChangesets(FToolMenuSection& Menu)
+{
+	Menu.AddMenuEntry(
+		TEXT("PlasticChangesetsWindow"),
+		LOCTEXT("PlasticChangesetsWindow", "View Changesets"),
+		LOCTEXT("PlasticChangesetsWindowTooltip", "Open the Changesets window."),
+		FSlateIcon(FAppStyle::GetAppStyleSetName(), "SourceControl.Actions.History"),
+		FUIAction(
+			FExecuteAction::CreateRaw(this, &FPlasticSourceControlMenu::OpenChangesetsWindow),
+			FCanExecuteAction()
+		)
+	);
+}
+
+void FPlasticSourceControlMenu::AddViewLocks(FToolMenuSection& Menu)
+{
+	const bool bVersionSupportsSmartLocks = FPlasticSourceControlModule::Get().GetProvider().GetPlasticScmVersion() >= PlasticSourceControlVersions::SmartLocks;
+
+	Menu.AddMenuEntry(
+		TEXT("PlasticLocksWindow"),
+		LOCTEXT("PlasticLocksWindow", "View Locks"),
+		LOCTEXT("PlasticLocksWindowTooltip", "Open the Locks window."),
+		FSlateIcon(FAppStyle::GetAppStyleSetName(), "PropertyWindow.Locked"),
+		FUIAction(
+			FExecuteAction::CreateRaw(this, &FPlasticSourceControlMenu::OpenLocksWindow),
+			FCanExecuteAction::CreateLambda([bVersionSupportsSmartLocks]() { return bVersionSupportsSmartLocks; })
+		)
+	);
+}
+
 #undef LOCTEXT_NAMESPACE
diff --git a/Engine/Plugins/Developer/PlasticSourceControl/Source/PlasticSourceControl/Private/PlasticSourceControlMenu.h b/Engine/Plugins/Developer/PlasticSourceControl/Source/PlasticSourceControl/Private/PlasticSourceControlMenu.h
--- a/Engine/Plugins/Developer/PlasticSourceControl/Source/PlasticSourceControl/Private/PlasticSourceControlMenu.h
+++ b/Engine/Plugins/Developer/PlasticSourceControl/Source/PlasticSourceControl/Private/PlasticSourceControlMenu.h
@@ -10,6 +10,8 @@
 class FMenuBuilder;
 struct FToolMenuSection;
 
+typedef TSharedRef<class FPlasticSourceControlLock, ESPMode::ThreadSafe> FPlasticSourceControlLockRef;
+
 /** Unity Version Control extension of the Source Control toolbar menu */
 class FPlasticSourceControlMenu
 {
@@ -29,13 +31,18 @@ public:
 	void VisitDocsURLClicked() const;
 	void VisitSupportURLClicked() const;
 	void VisitLockRulesURLClicked(const FString InOrganizationName) const;
+	void OpenDesktopApplication() const;
 	void OpenBranchesWindow() const;
+	void OpenChangesetsWindow() const;
+	void OpenLocksWindow() const;
 
 private:
 	bool IsSourceControlConnected() const;
 
 	void AddMenuExtension(FToolMenuSection& Menu);
 	void AddViewBranches(FToolMenuSection& Menu);
+	void AddViewChangesets(FToolMenuSection& Menu);
+	void AddViewLocks(FToolMenuSection& Menu);
 
 	/** Extends the UE5 toolbar with a status bar widget to display the current branch and open the branch tab */
 	void ExtendToolbarWithStatusBarWidget();
@@ -47,11 +54,11 @@ private:
 	/** Called to generate concert asset context menu. */
 	void GeneratePlasticAssetContextMenu(FMenuBuilder& MenuBuilder, TArray<FAssetData> InAssetObjectPaths);
 
-	bool CanRemoveLocks(TArray<FAssetData> InAssetObjectPaths) const;
-	bool CanReleaseLocks(TArray<FAssetData> InAssetObjectPaths) const;
-	void ExecuteRemoveLocks(TArray<FAssetData> InAssetObjectPaths);
-	void ExecuteReleaseLocks(TArray<FAssetData> InAssetObjectPaths);
-	void ExecuteUnlock(const TArray<FAssetData>& InAssetObjectPaths, const bool bInRemove);
+	bool CanRemoveLocks(TArray<FPlasticSourceControlLockRef> InSelectedLocks) const;
+	bool CanReleaseLocks(TArray<FPlasticSourceControlLockRef> InSelectedLocks) const;
+	void ExecuteRemoveLocks(TArray<FPlasticSourceControlLockRef> InSelectedLocks);
+	void ExecuteReleaseLocks(TArray<FPlasticSourceControlLockRef> InSelectedLocks);
+	void ExecuteUnlock(TArray<FPlasticSourceControlLockRef>&& InSelectedLocks, const bool bInRemove);
 
 private:
 	/** Tracks if the menu extension has been registered with the editor or not */
diff --git a/Engine/Plugins/Developer/PlasticSourceControl/Source/PlasticSourceControl/Private/PlasticSourceControlModule.cpp b/Engine/Plugins/Developer/PlasticSourceControl/Source/PlasticSourceControl/Private/PlasticSourceControlModule.cpp
--- a/Engine/Plugins/Developer/PlasticSourceControl/Source/PlasticSourceControl/Private/PlasticSourceControlModule.cpp
+++ b/Engine/Plugins/Developer/PlasticSourceControl/Source/PlasticSourceControl/Private/PlasticSourceControlModule.cpp
@@ -19,8 +19,10 @@ void FPlasticSourceControlModule::StartupModule()
 	// Bind our source control provider to the editor
 	IModularFeatures::Get().RegisterModularFeature("SourceControl", &PlasticSourceControlProvider);
 
-	/// Register our tab Window here as it needs to be ready for the editor to reload at startup
+	/// Register our tab windows here as they needs to be ready for the editor to reload at startup
 	PlasticSourceControlBranchesWindow.Register();
+	PlasticSourceControlChangesetsWindow.Register();
+	PlasticSourceControlLocksWindow.Register();
 }
 
 void FPlasticSourceControlModule::ShutdownModule()
@@ -29,6 +31,8 @@ void FPlasticSourceControlModule::ShutdownModule()
 	PlasticSourceControlProvider.Close();
 
 	PlasticSourceControlBranchesWindow.Unregister();
+	PlasticSourceControlChangesetsWindow.Unregister();
+	PlasticSourceControlLocksWindow.Unregister();
 
 	// unbind provider from editor
 	IModularFeatures::Get().UnregisterModularFeature("SourceControl", &PlasticSourceControlProvider);
diff --git a/Engine/Plugins/Developer/PlasticSourceControl/Source/PlasticSourceControl/Private/PlasticSourceControlModule.h b/Engine/Plugins/Developer/PlasticSourceControl/Source/PlasticSourceControl/Private/PlasticSourceControlModule.h
--- a/Engine/Plugins/Developer/PlasticSourceControl/Source/PlasticSourceControl/Private/PlasticSourceControlModule.h
+++ b/Engine/Plugins/Developer/PlasticSourceControl/Source/PlasticSourceControl/Private/PlasticSourceControlModule.h
@@ -9,6 +9,8 @@
 #include "PlasticSourceControlWorkspaceCreation.h"
 
 #include "PlasticSourceControlBranchesWindow.h"
+#include "PlasticSourceControlChangesetsWindow.h"
+#include "PlasticSourceControlLocksWindow.h"
 
 /**
  * PlasticSourceControl is the official Unity Version Control Plugin for Unreal Engine
@@ -43,6 +45,16 @@ public:
 		return PlasticSourceControlBranchesWindow;
 	}
 
+	FPlasticSourceControlChangesetsWindow& GetChangesetsWindow()
+	{
+		return PlasticSourceControlChangesetsWindow;
+	}
+
+	FPlasticSourceControlLocksWindow& GetLocksWindow()
+	{
+		return PlasticSourceControlLocksWindow;
+	}
+
 	/**
 	 * Singleton-like access to this module's interface.  This is just for convenience!
 	 * Beware of calling this during the shutdown phase, though.  Your module might have been unloaded already.
@@ -67,7 +79,10 @@ private:
 	/** The Plastic source control provider */
 	FPlasticSourceControlProvider PlasticSourceControlProvider;
 
+	/** Dockable windows adding advanced features to the plugin */
 	FPlasticSourceControlBranchesWindow PlasticSourceControlBranchesWindow;
+	FPlasticSourceControlChangesetsWindow PlasticSourceControlChangesetsWindow;
+	FPlasticSourceControlLocksWindow PlasticSourceControlLocksWindow;
 
 	/** Logic to create a new workspace */
 	FPlasticSourceControlWorkspaceCreation PlasticSourceControlWorkspaceCreation;
diff --git a/Engine/Plugins/Developer/PlasticSourceControl/Source/PlasticSourceControl/Private/PlasticSourceControlOperations.cpp b/Engine/Plugins/Developer/PlasticSourceControl/Source/PlasticSourceControl/Private/PlasticSourceControlOperations.cpp
--- a/Engine/Plugins/Developer/PlasticSourceControl/Source/PlasticSourceControl/Private/PlasticSourceControlOperations.cpp
+++ b/Engine/Plugins/Developer/PlasticSourceControl/Source/PlasticSourceControl/Private/PlasticSourceControlOperations.cpp
@@ -4,7 +4,9 @@
 
 #include "PackageUtils.h"
 #include "PlasticSourceControlBranch.h"
+#include "PlasticSourceControlChangeset.h"
 #include "PlasticSourceControlCommand.h"
+#include "PlasticSourceControlLock.h"
 #include "PlasticSourceControlModule.h"
 #include "PlasticSourceControlParsers.h"
 #include "PlasticSourceControlProvider.h"
@@ -46,14 +48,18 @@ void IPlasticSourceControlWorker::RegisterWorkers(FPlasticSourceControlProvider&
 	PlasticSourceControlProvider.RegisterWorker("Revert", FGetPlasticSourceControlWorker::CreateStatic(&InstantiateWorker<FPlasticRevertWorker>));
 	PlasticSourceControlProvider.RegisterWorker("RevertUnchanged", FGetPlasticSourceControlWorker::CreateStatic(&InstantiateWorker<FPlasticRevertUnchangedWorker>));
 	PlasticSourceControlProvider.RegisterWorker("RevertAll", FGetPlasticSourceControlWorker::CreateStatic(&InstantiateWorker<FPlasticRevertAllWorker>));
+	PlasticSourceControlProvider.RegisterWorker("RevertToRevision", FGetPlasticSourceControlWorker::CreateStatic(&InstantiateWorker<FPlasticRevertToRevisionWorker>));
 	PlasticSourceControlProvider.RegisterWorker("SwitchToPartialWorkspace", FGetPlasticSourceControlWorker::CreateStatic(&InstantiateWorker<FPlasticSwitchToPartialWorkspaceWorker>));
+	PlasticSourceControlProvider.RegisterWorker("GetLocks", FGetPlasticSourceControlWorker::CreateStatic(&InstantiateWorker<FPlasticGetLocksWorker>));
 	PlasticSourceControlProvider.RegisterWorker("Unlock", FGetPlasticSourceControlWorker::CreateStatic(&InstantiateWorker<FPlasticUnlockWorker>));
 	PlasticSourceControlProvider.RegisterWorker("GetBranches", FGetPlasticSourceControlWorker::CreateStatic(&InstantiateWorker<FPlasticGetBranchesWorker>));
-	PlasticSourceControlProvider.RegisterWorker("Switch", FGetPlasticSourceControlWorker::CreateStatic(&InstantiateWorker<FPlasticSwitchToBranchWorker>));
+	PlasticSourceControlProvider.RegisterWorker("Switch", FGetPlasticSourceControlWorker::CreateStatic(&InstantiateWorker<FPlasticSwitchWorker>));
 	PlasticSourceControlProvider.RegisterWorker("Merge", FGetPlasticSourceControlWorker::CreateStatic(&InstantiateWorker<FPlasticMergeBranchWorker>));
 	PlasticSourceControlProvider.RegisterWorker("CreateBranch", FGetPlasticSourceControlWorker::CreateStatic(&InstantiateWorker<FPlasticCreateBranchWorker>));
 	PlasticSourceControlProvider.RegisterWorker("RenameBranch", FGetPlasticSourceControlWorker::CreateStatic(&InstantiateWorker<FPlasticRenameBranchWorker>));
 	PlasticSourceControlProvider.RegisterWorker("DeleteBranches", FGetPlasticSourceControlWorker::CreateStatic(&InstantiateWorker<FPlasticDeleteBranchesWorker>));
+	PlasticSourceControlProvider.RegisterWorker("GetChangesets", FGetPlasticSourceControlWorker::CreateStatic(&InstantiateWorker<FPlasticGetChangesetsWorker>));
+	PlasticSourceControlProvider.RegisterWorker("GetChangesetFiles", FGetPlasticSourceControlWorker::CreateStatic(&InstantiateWorker<FPlasticGetChangesetFilesWorker>));
 	PlasticSourceControlProvider.RegisterWorker("MakeWorkspace", FGetPlasticSourceControlWorker::CreateStatic(&InstantiateWorker<FPlasticMakeWorkspaceWorker>));
 	PlasticSourceControlProvider.RegisterWorker("Sync", FGetPlasticSourceControlWorker::CreateStatic(&InstantiateWorker<FPlasticSyncWorker>));
 	PlasticSourceControlProvider.RegisterWorker("SyncAll", FGetPlasticSourceControlWorker::CreateStatic(&InstantiateWorker<FPlasticSyncWorker>));
@@ -91,6 +97,16 @@ FName FPlasticSyncAll::GetName() const
 	return "SyncAll";
 }
 
+FText FPlasticSyncAll::GetInProgressString() const
+{
+	if (!Revision.IsEmpty())
+	{
+		return FText::Format(LOCTEXT("SourceControl_SyncAll", "Switching the workspace to changeset {0}..."), FText::FromString(Revision));
+	}
+
+	return FSync::GetInProgressString();
+}
+
 FName FPlasticRevertAll::GetName() const
 {
 	return "RevertAll";
@@ -101,6 +117,16 @@ FText FPlasticRevertAll::GetInProgressString() const
 	return LOCTEXT("SourceControl_RevertAll", "Reverting checked-out file(s) in Revision Control...");
 }
 
+FName FPlasticRevertToRevision::GetName() const
+{
+	return "RevertToRevision";
+}
+
+FText FPlasticRevertToRevision::GetInProgressString() const
+{
+	return FText::Format(LOCTEXT("SourceControl_RevertToRevision", "Reverting selected file(s) to revision {0}..."), FText::AsNumber(ChangesetId));
+}
+
 FName FPlasticMakeWorkspace::GetName() const
 {
 	return "MakeWorkspace";
@@ -121,6 +147,16 @@ FText FPlasticSwitchToPartialWorkspace::GetInProgressString() const
 	return LOCTEXT("SourceControl_SwitchToPartialWorkspace", "Switching to a Partial/Gluon Workspace...");
 }
 
+FName FPlasticGetLocks::GetName() const
+{
+	return "GetLocks";
+}
+
+FText FPlasticGetLocks::GetInProgressString() const
+{
+	return LOCTEXT("SourceControl_GetLocks", "Getting the list of locks...");
+}
+
 FName FPlasticUnlock::GetName() const
 {
 	return "Unlock";
@@ -144,14 +180,14 @@ FText FPlasticGetBranches::GetInProgressString() const
 	return LOCTEXT("SourceControl_GetBranches", "Getting the list of branches...");
 }
 
-FName FPlasticSwitchToBranch::GetName() const
+FName FPlasticSwitch::GetName() const
 {
 	return "Switch";
 }
 
-FText FPlasticSwitchToBranch::GetInProgressString() const
+FText FPlasticSwitch::GetInProgressString() const
 {
-	return LOCTEXT("SourceControl_SwitchToBranch", "Switching the workspace to another branch...");
+	return FText::Format(LOCTEXT("SourceControl_SwitchToBranch", "Switching the workspace to branch {0}..."), FText::FromString(BranchName));
 }
 
 FName FPlasticMergeBranch::GetName() const
@@ -161,7 +197,7 @@ FName FPlasticMergeBranch::GetName() const
 
 FText FPlasticMergeBranch::GetInProgressString() const
 {
-	return LOCTEXT("SourceControl_MergeBranch", "Merging branch to the current one...");
+	return FText::Format(LOCTEXT("SourceControl_MergeBranch", "Merging branch {0} to the current one..."), FText::FromString(BranchName));
 }
 
 FName FPlasticCreateBranch::GetName() const
@@ -171,7 +207,7 @@ FName FPlasticCreateBranch::GetName() const
 
 FText FPlasticCreateBranch::GetInProgressString() const
 {
-	return LOCTEXT("SourceControl_CreateBranch", "Creating new child branch...");
+	return FText::Format(LOCTEXT("SourceControl_CreateBranch", "Creating new child branch {0}..."), FText::FromString(BranchName));
 }
 
 
@@ -182,7 +218,7 @@ FName FPlasticRenameBranch::GetName() const
 
 FText FPlasticRenameBranch::GetInProgressString() const
 {
-	return LOCTEXT("SourceControl_RenameBranch", "Renaming a branch...");
+	return FText::Format(LOCTEXT("SourceControl_RenameBranch", "Renaming the branch {0} to {1}..."), FText::FromString(OldName), FText::FromString(NewName));
 }
 
 FName FPlasticDeleteBranches::GetName() const
@@ -192,7 +228,27 @@ FName FPlasticDeleteBranches::GetName() const
 
 FText FPlasticDeleteBranches::GetInProgressString() const
 {
-	return LOCTEXT("SourceControl_DeleteBranches", "Deleting branch(es)...");
+	return FText::Format(LOCTEXT("SourceControl_DeleteBranches", "Deleting branch(es) {0}..."), FText::FromString(BranchNames[0]));
+}
+
+FName FPlasticGetChangesets::GetName() const
+{
+	return "GetChangesets";
+}
+
+FText FPlasticGetChangesets::GetInProgressString() const
+{
+	return LOCTEXT("SourceControl_GetChangesets", "Getting the list of changesets...");
+}
+
+FName FPlasticGetChangesetFiles::GetName() const
+{
+	return "GetChangesetFiles";
+}
+
+FText FPlasticGetChangesetFiles::GetInProgressString() const
+{
+	return FText::Format(LOCTEXT("SourceControl_GetChangesetFiles", "Getting the list of files in changeset {0}..."), FText::AsNumber(Changeset->ChangesetId));
 }
 
 
@@ -226,7 +282,7 @@ bool FPlasticConnectWorker::Execute(FPlasticSourceControlCommand& InCommand)
 		if (InCommand.bCommandSuccessful)
 		{
 			// Get current branch, repository, server URL and check the connection
-			InCommand.bCommandSuccessful = PlasticSourceControlUtils::RunCheckConnection(InCommand.BranchName, InCommand.RepositoryName, InCommand.ServerUrl, InCommand.InfoMessages, InCommand.ErrorMessages);
+			InCommand.bCommandSuccessful = PlasticSourceControlUtils::RunCheckConnection(InCommand.WorkspaceSelector, InCommand.BranchName, InCommand.RepositoryName, InCommand.ServerUrl, InCommand.InfoMessages, InCommand.ErrorMessages);
 			if (InCommand.bCommandSuccessful)
 			{
 				InCommand.InfoMessages.Add(TEXT("Connected successfully"));
@@ -239,7 +295,7 @@ bool FPlasticConnectWorker::Execute(FPlasticSourceControlCommand& InCommand)
 					PlasticSourceControlShell::SetShellIsWarmedUp();
 					TArray<FString> ContentDir;
 					ContentDir.Add(FPaths::ConvertRelativePathToFull(FPaths::ProjectContentDir()));
-					PlasticSourceControlUtils::RunUpdateStatus(ContentDir, PlasticSourceControlUtils::EStatusSearchType::All, false, InCommand.ErrorMessages, States, InCommand.ChangesetNumber, InCommand.BranchName);
+					PlasticSourceControlUtils::RunUpdateStatus(ContentDir, PlasticSourceControlUtils::EStatusSearchType::All, false, InCommand.ErrorMessages, States, InCommand.ChangesetNumber);
 				}
 			}
 			else
@@ -291,15 +347,14 @@ static void UpdateChangelistState(FPlasticSourceControlProvider& SCCProvider, co
 
 		for (const FPlasticSourceControlState& InState : InStates)
 		{
-			// Note: cannot use IsModified() because cm cannot yet handle local modifications in changelists, only the GUI can
-			if (!InState.IsCheckedOutImplementation())
+			if (!InState.IsPendingChanges())
 			{
 				continue;
 			}
 
 			// Add a shared reference to the state of the file, that will then be updated by PlasticSourceControlUtils::UpdateCachedStates()
 			TSharedRef<FPlasticSourceControlState, ESPMode::ThreadSafe> State = SCCProvider.GetStateInternal(InState.GetFilename());
-			ChangelistState->Files.Add(State);
+			ChangelistState->Files.AddUnique(State);
 
 			// Keep the changelist stored with cached file state in sync with the actual changelist that owns this file.
 			State->Changelist = InChangelist;
@@ -329,7 +384,8 @@ bool FPlasticCheckOutWorker::Execute(FPlasticSourceControlCommand& InCommand)
 	}
 
 	// now update the status of our files
-	PlasticSourceControlUtils::RunUpdateStatus(InCommand.Files, PlasticSourceControlUtils::EStatusSearchType::ControlledOnly, false, InCommand.ErrorMessages, States, InCommand.ChangesetNumber, InCommand.BranchName);
+	PlasticSourceControlUtils::InvalidateLocksCache();
+	PlasticSourceControlUtils::RunUpdateStatus(InCommand.Files, PlasticSourceControlUtils::EStatusSearchType::ControlledOnly, false, InCommand.ErrorMessages, States, InCommand.ChangesetNumber);
 
 	return InCommand.bCommandSuccessful;
 }
@@ -364,21 +420,21 @@ bool DeleteChangelist(const FPlasticSourceControlProvider& PlasticSourceControlP
 	}
 }
 
-TArray<FString> FileNamesFromFileStates(const TArray<FSourceControlStateRef>& InFileStates)
+static TArray<FString> FileNamesFromFileStates(const TArray<FSourceControlStateRef>& InFileStates)
 {
 	TArray<FString> Files;
-
+	Files.Reserve(InFileStates.Num());
 	for (const FSourceControlStateRef& FileState : InFileStates)
 	{
 		Files.Add(FileState->GetFilename());
 	}
-
 	return Files;
 }
 
-TArray<FString> GetFilesFromCommand(FPlasticSourceControlProvider& PlasticSourceControlProvider, FPlasticSourceControlCommand& InCommand)
+static TArray<FString> GetFilesFromCommand(FPlasticSourceControlProvider& PlasticSourceControlProvider, FPlasticSourceControlCommand& InCommand)
 {
 	TArray<FString> Files;
+
 	if (InCommand.Changelist.IsInitialized() && InCommand.Files.IsEmpty())
 	{
 		TSharedRef<FPlasticSourceControlChangelistState, ESPMode::ThreadSafe> ChangelistState = PlasticSourceControlProvider.GetStateInternal(InCommand.Changelist);
@@ -407,61 +463,83 @@ bool FPlasticCheckInWorker::Execute(FPlasticSourceControlCommand& InCommand)
 
 	const TArray<FString> Files = GetFilesFromCommand(GetProvider(), InCommand);
 
-	if (Files.Num() > 0)
+	if (Files.Num() == 0)
 	{
-		if (InCommand.Changelist.IsInitialized())
+		UE_LOG(LogSourceControl, Warning, TEXT("Checkin: No files provided"));
+		return false;
+	}
+
+	// Add Private files to source control to be able to check them in
+	TArray<FString> FilesToAdd;
+	for (const FString& File : Files)
+	{
+		TSharedRef<FPlasticSourceControlState, ESPMode::ThreadSafe> State = GetProvider().GetStateInternal(File);
+		if (!State->IsSourceControlled())
 		{
-			TSharedRef<FPlasticSourceControlChangelistState, ESPMode::ThreadSafe> ChangelistState = GetProvider().GetStateInternal(InCommand.Changelist);
+			FilesToAdd.Add(File);
+		}
+	}
+	if (FilesToAdd.Num() > 0)
+	{
+		UE_LOG(LogSourceControl, Log, TEXT("CheckIn: Adding %d file(s) to source control"), FilesToAdd.Num());
+		if (!GetProvider().IsPartialWorkspace())
+		{
+			InCommand.bCommandSuccessful = PlasticSourceControlUtils::RunCommand(TEXT("add"), TArray<FString>(), FilesToAdd, InCommand.InfoMessages, InCommand.ErrorMessages);
+		}
+		else
+		{
+			InCommand.bCommandSuccessful = PlasticSourceControlUtils::RunCommand(TEXT("partial add"), TArray<FString>(), FilesToAdd, InCommand.InfoMessages, InCommand.ErrorMessages);
+		}
+	}
 
-			if (Description.IsEmpty())
-			{
-				Description = ChangelistState->GetDescriptionText();
-			}
+	if (InCommand.Changelist.IsInitialized())
+	{
+		TSharedRef<FPlasticSourceControlChangelistState, ESPMode::ThreadSafe> ChangelistState = GetProvider().GetStateInternal(InCommand.Changelist);
 
-			InChangelist = InCommand.Changelist;
+		if (Description.IsEmpty())
+		{
+			Description = ChangelistState->GetDescriptionText();
 		}
 
-		UE_LOG(LogSourceControl, Verbose, TEXT("CheckIn: %d file(s) Description: '%s'"), Files.Num(), *Description.ToString());
+		InChangelist = InCommand.Changelist;
+	}
 
-		// make a temp file to place our commit message in
-		const FScopedTempFile CommitMsgFile(Description);
-		if (!CommitMsgFile.GetFilename().IsEmpty())
+	UE_LOG(LogSourceControl, Verbose, TEXT("CheckIn: %d file(s) Description: '%s'"), Files.Num(), *Description.ToString());
+
+	// make a temp file to place our commit message in
+	const FScopedTempFile CommitMsgFile(Description);
+	if (!CommitMsgFile.GetFilename().IsEmpty())
+	{
+		TArray<FString> Parameters;
+		Parameters.Add(FString::Printf(TEXT("--commentsfile=\"%s\""), *CommitMsgFile.GetFilename()));
+		Parameters.Add(TEXT("--all"));			// Also files Changed (not CheckedOut) and Moved/Deleted Locally
+		if (!GetProvider().IsPartialWorkspace())
 		{
-			TArray<FString> Parameters;
-			Parameters.Add(FString::Printf(TEXT("--commentsfile=\"%s\""), *CommitMsgFile.GetFilename()));
-			if (!GetProvider().IsPartialWorkspace())
-			{
-				Parameters.Add(TEXT("--all")); // Also files Changed (not CheckedOut) and Moved/Deleted Locally
 			//  NOTE: --update added as #23 but removed as #32 because most assets are locked by the Unreal Editor
-			//  Parameters.Add(TEXT("--update")); // Processes the update-merge automatically if it eventually happens.
-				InCommand.bCommandSuccessful = PlasticSourceControlUtils::RunCommand(TEXT("checkin"), Parameters, Files, InCommand.InfoMessages, InCommand.ErrorMessages);
-			}
-			else
-			{
-				Parameters.Add(TEXT("--applychanged")); // Also files Changed (not CheckedOut) and Moved/Deleted Locally
-				InCommand.bCommandSuccessful = PlasticSourceControlUtils::RunCommand(TEXT("partial checkin"), Parameters, Files, InCommand.InfoMessages, InCommand.ErrorMessages);
-			}
-			if (InCommand.bCommandSuccessful)
-			{
-				Operation->SetSuccessMessage(PlasticSourceControlParsers::ParseCheckInResults(InCommand.InfoMessages));
-				UE_LOG(LogSourceControl, Log, TEXT("CheckIn successful"));
-			}
-
-			if (InChangelist.IsInitialized() && !InChangelist.IsDefault())
-			{
-				// NOTE: we need to explicitly delete persistent changelists when we submit its content, except for the Default changelist
-				DeleteChangelist(GetProvider(), InChangelist, InCommand.InfoMessages, InCommand.ErrorMessages);
-			}
+		//  Parameters.Add(TEXT("--update")); // Processes the update-merge automatically if it eventually happens.
+			InCommand.bCommandSuccessful = PlasticSourceControlUtils::RunCommand(TEXT("checkin"), Parameters, Files, InCommand.InfoMessages, InCommand.ErrorMessages);
+		}
+		else
+		{
+			InCommand.bCommandSuccessful = PlasticSourceControlUtils::RunCommand(TEXT("partial checkin"), Parameters, Files, InCommand.InfoMessages, InCommand.ErrorMessages);
+		}
+		if (InCommand.bCommandSuccessful)
+		{
+			Operation->SetSuccessMessage(PlasticSourceControlParsers::ParseCheckInResults(InCommand.InfoMessages));
+			UE_LOG(LogSourceControl, Log, TEXT("CheckIn successful"));
 		}
 
-		// now update the status of our files
-		PlasticSourceControlUtils::RunUpdateStatus(Files, PlasticSourceControlUtils::EStatusSearchType::ControlledOnly, false, InCommand.ErrorMessages, States, InCommand.ChangesetNumber, InCommand.BranchName);
-	}
-	else
-	{
-		UE_LOG(LogSourceControl, Warning, TEXT("Checkin: No files provided"));
+		if (InChangelist.IsInitialized() && !InChangelist.IsDefault())
+		{
+			// NOTE: we need to explicitly delete persistent changelists when we submit its content, except for the Default changelist
+			DeleteChangelist(GetProvider(), InChangelist, InCommand.InfoMessages, InCommand.ErrorMessages);
+		}
 	}
 
+	// now update the status of our files
+	PlasticSourceControlUtils::InvalidateLocksCache();
+	PlasticSourceControlUtils::RunUpdateStatus(Files, PlasticSourceControlUtils::EStatusSearchType::ControlledOnly, false, InCommand.ErrorMessages, States, InCommand.ChangesetNumber);
+
 	return InCommand.bCommandSuccessful;
 }
 
@@ -559,7 +637,7 @@ bool FPlasticMarkForAddWorker::Execute(FPlasticSourceControlCommand& InCommand)
 		}
 
 		// now update the status of our files
-		PlasticSourceControlUtils::RunUpdateStatus(InCommand.Files, PlasticSourceControlUtils::EStatusSearchType::ControlledOnly, false, InCommand.ErrorMessages, States, InCommand.ChangesetNumber, InCommand.BranchName);
+		PlasticSourceControlUtils::RunUpdateStatus(InCommand.Files, PlasticSourceControlUtils::EStatusSearchType::ControlledOnly, false, InCommand.ErrorMessages, States, InCommand.ChangesetNumber);
 	}
 	else
 	{
@@ -600,7 +678,7 @@ bool FPlasticDeleteWorker::Execute(FPlasticSourceControlCommand& InCommand)
 	}
 
 	// now update the status of our files
-	PlasticSourceControlUtils::RunUpdateStatus(InCommand.Files, PlasticSourceControlUtils::EStatusSearchType::ControlledOnly, false, InCommand.ErrorMessages, States, InCommand.ChangesetNumber, InCommand.BranchName);
+	PlasticSourceControlUtils::RunUpdateStatus(InCommand.Files, PlasticSourceControlUtils::EStatusSearchType::ControlledOnly, false, InCommand.ErrorMessages, States, InCommand.ChangesetNumber);
 
 	return InCommand.bCommandSuccessful;
 }
@@ -645,15 +723,28 @@ bool FPlasticRevertWorker::Execute(FPlasticSourceControlCommand& InCommand)
 		return false;
 	}
 
-	const TArray<FString> Files = GetFilesFromCommand(GetProvider(), InCommand);
+	TArray<FString> Files = GetFilesFromCommand(GetProvider(), InCommand);
 
 	TArray<FString> LocallyChangedFiles;
 	TArray<FString> CheckedOutFiles;
 
-	for (const FString& File : Files)
+	int32 i = 0;
+	while (i < Files.Num())
 	{
+		FString& File = Files[i];
+
 		const TSharedRef<const FPlasticSourceControlState, ESPMode::ThreadSafe> State = GetProvider().GetStateInternal(File);
 
+		// Remove files that cannot be reverted from the list (e.g. Private files)
+		if (!State->CanRevert())
+		{
+			// Private files are not under source control, so we can't revert them
+			UE_LOG(LogSourceControl, Log, TEXT("File '%s' cannot be reverted"), *File);
+			Files.RemoveAt(i);
+			continue;
+		}
+		i++;
+
 		if (State->WorkspaceState == EWorkspaceState::Changed)
 		{
 			LocallyChangedFiles.Add(State->LocalFilename);
@@ -720,7 +811,8 @@ bool FPlasticRevertWorker::Execute(FPlasticSourceControlCommand& InCommand)
 	// NOTE: optim, in UE4 there was no need to update the status of our files since this is done immediately after by the Editor, except now that we are using changelists
 	// update the status of our files: need to check for local changes in case of a SoftRevert
 	const PlasticSourceControlUtils::EStatusSearchType SearchType = bIsSoftRevert ? PlasticSourceControlUtils::EStatusSearchType::All : PlasticSourceControlUtils::EStatusSearchType::ControlledOnly;
-	PlasticSourceControlUtils::RunUpdateStatus(Files, SearchType, false, InCommand.ErrorMessages, States, InCommand.ChangesetNumber, InCommand.BranchName);
+	PlasticSourceControlUtils::InvalidateLocksCache();
+	PlasticSourceControlUtils::RunUpdateStatus(Files, SearchType, false, InCommand.ErrorMessages, States, InCommand.ChangesetNumber);
 
 	return InCommand.bCommandSuccessful;
 }
@@ -746,6 +838,26 @@ bool FPlasticRevertWorker::UpdateStates()
 	return PlasticSourceControlUtils::UpdateCachedStates(MoveTemp(States));
 }
 
+
+static TArray<FString> GetFilesCheckedOut(FPlasticSourceControlProvider& InSourceControlProvider, const TArray<FString>& InFiles)
+{
+	TArray<FString> MoveableFiles;
+	MoveableFiles.Reserve(InFiles.Num());
+	for (const auto& File : InFiles)
+	{
+		TSharedRef<FPlasticSourceControlState, ESPMode::ThreadSafe> FileState = InSourceControlProvider.GetStateInternal(File);
+		if (FileState->IsCheckedOutImplementation())
+		{
+			MoveableFiles.Add(File);
+		}
+		else
+		{
+			UE_LOG(LogSourceControl, Warning, TEXT("GetFilesCheckedOut: File '%s' is not checked out and cannot be moved to another changelist"), *File);
+		}
+	}
+	return MoveableFiles;
+}
+
 FName FPlasticRevertUnchangedWorker::GetName() const
 {
 	return "RevertUnchanged";
@@ -760,7 +872,14 @@ bool FPlasticRevertUnchangedWorker::Execute(FPlasticSourceControlCommand& InComm
 	TArray<FString> Parameters;
 	Parameters.Add(TEXT("-R"));
 
-	TArray<FString> Files = GetFilesFromCommand(GetProvider(), InCommand);
+	// Filter out files that we shouldn't try to uncheckout!
+	TArray<FString> FilesFromCommand = GetFilesFromCommand(GetProvider(), InCommand);
+	TArray<FString> Files = GetFilesCheckedOut(GetProvider(), FilesFromCommand);
+	if (Files.Num() == 0 && FilesFromCommand.Num() > 0)
+	{
+		// If the command specified a set of files but none are checked out, do nothing
+		return true;
+	}
 
 	// revert the checkout of all unchanged files recursively
 	InCommand.bCommandSuccessful = PlasticSourceControlUtils::RunCommand(TEXT("uncounchanged"), Parameters, Files, InCommand.InfoMessages, InCommand.ErrorMessages);
@@ -770,7 +889,8 @@ bool FPlasticRevertUnchangedWorker::Execute(FPlasticSourceControlCommand& InComm
 	{
 		Files.Add(FPaths::ConvertRelativePathToFull(FPaths::ProjectContentDir()));
 	}
-	PlasticSourceControlUtils::RunUpdateStatus(Files, PlasticSourceControlUtils::EStatusSearchType::ControlledOnly, false, InCommand.ErrorMessages, States, InCommand.ChangesetNumber, InCommand.BranchName);
+	PlasticSourceControlUtils::InvalidateLocksCache();
+	PlasticSourceControlUtils::RunUpdateStatus(Files, PlasticSourceControlUtils::EStatusSearchType::ControlledOnly, false, InCommand.ErrorMessages, States, InCommand.ChangesetNumber);
 
 	return InCommand.bCommandSuccessful;
 }
@@ -782,7 +902,7 @@ bool FPlasticRevertUnchangedWorker::UpdateStates()
 	// Update affected changelists if any
 	for (const FPlasticSourceControlState& NewState : States)
 	{
-		if (!NewState.IsCheckedOutImplementation())
+		if (!NewState.IsPendingChanges())
 		{
 			TSharedRef<FPlasticSourceControlState, ESPMode::ThreadSafe> State = GetProvider().GetStateInternal(NewState.GetFilename());
 			if (State->Changelist.IsInitialized())
@@ -816,7 +936,7 @@ bool FPlasticRevertAllWorker::Execute(FPlasticSourceControlCommand& InCommand)
 		TArray<FPlasticSourceControlState> TempStates;
 		TArray<FString> ContentDir;
 		ContentDir.Add(FPaths::ConvertRelativePathToFull(FPaths::ProjectContentDir()));
-		PlasticSourceControlUtils::RunUpdateStatus(ContentDir, PlasticSourceControlUtils::EStatusSearchType::All, false, InCommand.ErrorMessages, TempStates, InCommand.ChangesetNumber, InCommand.BranchName);
+		PlasticSourceControlUtils::RunUpdateStatus(ContentDir, PlasticSourceControlUtils::EStatusSearchType::All, false, InCommand.ErrorMessages, TempStates, InCommand.ChangesetNumber);
 
 		for (auto& State : TempStates)
 		{
@@ -863,7 +983,8 @@ bool FPlasticRevertAllWorker::Execute(FPlasticSourceControlCommand& InCommand)
 	// now update the status of the updated files
 	if (Operation->UpdatedFiles.Num())
 	{
-		PlasticSourceControlUtils::RunUpdateStatus(Operation->UpdatedFiles, PlasticSourceControlUtils::EStatusSearchType::ControlledOnly, false, InCommand.ErrorMessages, States, InCommand.ChangesetNumber, InCommand.BranchName);
+		PlasticSourceControlUtils::InvalidateLocksCache();
+		PlasticSourceControlUtils::RunUpdateStatus(Operation->UpdatedFiles, PlasticSourceControlUtils::EStatusSearchType::ControlledOnly, false, InCommand.ErrorMessages, States, InCommand.ChangesetNumber);
 	}
 
 	return InCommand.bCommandSuccessful;
@@ -876,8 +997,7 @@ bool FPlasticRevertAllWorker::UpdateStates()
 	// Update affected changelists if any
 	for (const FPlasticSourceControlState& NewState : States)
 	{
-		// TODO: also detect files that were added and are now private! Should be removed as well from their changelist
-		if (!NewState.IsCheckedOutImplementation())
+		if (!NewState.IsPendingChanges())
 		{
 			TSharedRef<FPlasticSourceControlState, ESPMode::ThreadSafe> State = GetProvider().GetStateInternal(NewState.GetFilename());
 			if (State->Changelist.IsInitialized())
@@ -894,6 +1014,43 @@ bool FPlasticRevertAllWorker::UpdateStates()
 	return PlasticSourceControlUtils::UpdateCachedStates(MoveTemp(States));
 }
 
+FName FPlasticRevertToRevisionWorker::GetName() const
+{
+	return "RevertToRevision";
+}
+
+bool FPlasticRevertToRevisionWorker::Execute(FPlasticSourceControlCommand& InCommand)
+{
+	TRACE_CPUPROFILER_EVENT_SCOPE(FPlasticRevertToRevisionWorker::Execute);
+
+	check(InCommand.Operation->GetName() == GetName());
+	TSharedRef<FPlasticRevertToRevision, ESPMode::ThreadSafe> Operation = StaticCastSharedRef<FPlasticRevertToRevision>(InCommand.Operation);
+
+	InCommand.bCommandSuccessful = true;
+	for (const FString& File : InCommand.Files)
+	{
+		TArray<FString> Parameters;
+		Parameters.Add(FString::Printf(TEXT("\"%s#cs:%d\""), *File, Operation->ChangesetId));
+
+		InCommand.bCommandSuccessful &= PlasticSourceControlUtils::RunCommand(TEXT("revert"), Parameters, TArray<FString>(), InCommand.InfoMessages, InCommand.ErrorMessages);
+	}
+
+	if (InCommand.bCommandSuccessful)
+	{
+		Operation->UpdatedFiles = InCommand.Files;
+	}
+
+	// now update the status of our files
+	PlasticSourceControlUtils::RunUpdateStatus(InCommand.Files, PlasticSourceControlUtils::EStatusSearchType::ControlledOnly, false, InCommand.ErrorMessages, States, InCommand.ChangesetNumber);
+
+	return InCommand.bCommandSuccessful;
+}
+
+bool FPlasticRevertToRevisionWorker::UpdateStates()
+{
+	return PlasticSourceControlUtils::UpdateCachedStates(MoveTemp(States));
+}
+
 FName FPlasticMakeWorkspaceWorker::GetName() const
 {
 	return "MakeWorkspace";
@@ -910,8 +1067,15 @@ bool FPlasticMakeWorkspaceWorker::Execute(FPlasticSourceControlCommand& InComman
 		TArray<FString> Parameters;
 		Parameters.Add(Operation->ServerUrl);
 		Parameters.Add(Operation->RepositoryName);
-		PlasticSourceControlUtils::RunCommand(TEXT("makerepository"), Parameters, TArray<FString>(), InCommand.InfoMessages, InCommand.ErrorMessages);
+		// Note: the whole operation should fail entirely if the repository creation failed (if the repository already exists, if the organization name is invalid, credential, autorizations etc.)
+		InCommand.bCommandSuccessful = PlasticSourceControlUtils::RunCommand(TEXT("makerepository"), Parameters, TArray<FString>(), InCommand.InfoMessages, InCommand.ErrorMessages);
+		// Specifically detect the specific error when the organization name is invalid, and add an more human readable message.
+		if (InCommand.ErrorMessages.Contains(TEXT("Can't resolve DNS entry for cloud.plasticscm.com")))
+		{
+			InCommand.ErrorMessages.Add(TEXT("Invalid cloud organization name."));
+		}
 	}
+	if (InCommand.bCommandSuccessful)
 	{
 		TArray<FString> Parameters;
 		Parameters.Add(Operation->WorkspaceName);
@@ -919,7 +1083,7 @@ bool FPlasticMakeWorkspaceWorker::Execute(FPlasticSourceControlCommand& InComman
 		Parameters.Add(FString::Printf(TEXT("--repository=rep:%s@repserver:%s"), *Operation->RepositoryName, *Operation->ServerUrl));
 		InCommand.bCommandSuccessful = PlasticSourceControlUtils::RunCommand(TEXT("makeworkspace"), Parameters, TArray<FString>(), InCommand.InfoMessages, InCommand.ErrorMessages);
 	}
-	if (Operation->bPartialWorkspace)
+	if (InCommand.bCommandSuccessful && Operation->bPartialWorkspace)
 	{
 		TArray<FString> Parameters;
 		Parameters.Add(TEXT("update"));
@@ -955,7 +1119,7 @@ bool FPlasticSwitchToPartialWorkspaceWorker::Execute(FPlasticSourceControlComman
 	{
 		TArray<FString> ProjectFiles;
 		ProjectFiles.Add(FPaths::ConvertRelativePathToFull(FPaths::GetProjectFilePath()));
-		InCommand.bCommandSuccessful = PlasticSourceControlUtils::RunUpdateStatus(ProjectFiles, PlasticSourceControlUtils::EStatusSearchType::ControlledOnly, false, InCommand.ErrorMessages, States, InCommand.ChangesetNumber, InCommand.BranchName);
+		PlasticSourceControlUtils::RunUpdateStatus(ProjectFiles, PlasticSourceControlUtils::EStatusSearchType::ControlledOnly, false, InCommand.ErrorMessages, States, InCommand.ChangesetNumber);
 	}
 
 	return InCommand.bCommandSuccessful;
@@ -966,6 +1130,37 @@ bool FPlasticSwitchToPartialWorkspaceWorker::UpdateStates()
 	return PlasticSourceControlUtils::UpdateCachedStates(MoveTemp(States));
 }
 
+FName FPlasticGetLocksWorker::GetName() const
+{
+	return "GetLocks";
+}
+
+bool FPlasticGetLocksWorker::Execute(FPlasticSourceControlCommand& InCommand)
+{
+	TRACE_CPUPROFILER_EVENT_SCOPE(FPlasticGetLocksWorker::Execute);
+
+	check(InCommand.Operation->GetName() == GetName());
+	TSharedRef<FPlasticGetLocks, ESPMode::ThreadSafe> Operation = StaticCastSharedRef<FPlasticGetLocks>(InCommand.Operation);
+
+	{
+		// In the View Locks window, always show locks for all destination branches
+		const bool bForAllDestBranches = true;
+		InCommand.bCommandSuccessful = PlasticSourceControlUtils::RunListLocks(GetProvider(), bForAllDestBranches, Operation->Locks);
+	}
+
+	{
+		FString RepositoryName, ServerUrl;
+		InCommand.bCommandSuccessful &= PlasticSourceControlUtils::GetWorkspaceInfo(InCommand.WorkspaceSelector, InCommand.BranchName, RepositoryName, ServerUrl, InCommand.ErrorMessages);
+	}
+
+	return InCommand.bCommandSuccessful;
+}
+
+bool FPlasticGetLocksWorker::UpdateStates()
+{
+	return false;
+}
+
 FName FPlasticUnlockWorker::GetName() const
 {
 	return "Unlock";
@@ -978,36 +1173,44 @@ bool FPlasticUnlockWorker::Execute(FPlasticSourceControlCommand& InCommand)
 	check(InCommand.Operation->GetName() == GetName());
 	TSharedRef<FPlasticUnlock, ESPMode::ThreadSafe> Operation = StaticCastSharedRef<FPlasticUnlock>(InCommand.Operation);
 
-	{
-		// retrieve the itemid of assets to unlock
-		FString ItemIds;
-		for (const FString& File : InCommand.Files)
-		{
-			const auto State = GetProvider().GetStateInternal(File);
-			if (State->LockedId != ISourceControlState::INVALID_REVISION)
-			{
-				ItemIds += FString::Printf(TEXT("itemid:%d "), State->LockedId);
-			}
-		}
+	InCommand.bCommandSuccessful = true;
 
+	TArray<FString> Branches;
+	for (const FPlasticSourceControlLockRef& Lock : Operation->Locks)
+	{
+		Branches.AddUnique(Lock->Branch);
+	}
+	for (const FString& Branch : Branches)
+	{
 		TArray<FString> Parameters;
 		Parameters.Add(TEXT("unlock"));
 		if (Operation->bRemove)
 			Parameters.Add(TEXT("--remove"));
-		Parameters.Add(ItemIds);
-		InCommand.bCommandSuccessful = PlasticSourceControlUtils::RunCommand(TEXT("lock"), Parameters, TArray<FString>(), InCommand.InfoMessages, InCommand.ErrorMessages);
-	}
 
-	{
-		InCommand.bCommandSuccessful = PlasticSourceControlUtils::RunUpdateStatus(InCommand.Files, PlasticSourceControlUtils::EStatusSearchType::ControlledOnly, false, InCommand.ErrorMessages, States, InCommand.ChangesetNumber, InCommand.BranchName);
+		Parameters.Add(FString::Printf(TEXT("--branch=%s"), *Branch));
+		for (const FPlasticSourceControlLockRef& Lock : Operation->Locks)
+		{
+			if (Lock->Branch == Branch)
+			{
+				Parameters.Add(FString::Printf(TEXT("itemid:%d "), Lock->ItemId));
+			}
+		}
+		InCommand.bCommandSuccessful &= PlasticSourceControlUtils::RunCommand(TEXT("lock"), Parameters, TArray<FString>(), InCommand.InfoMessages, InCommand.ErrorMessages);
 	}
 
+	// now update the status of our files
+	PlasticSourceControlUtils::InvalidateLocksCache();
+	PlasticSourceControlUtils::RunUpdateStatus(InCommand.Files, PlasticSourceControlUtils::EStatusSearchType::ControlledOnly, false, InCommand.ErrorMessages, States, InCommand.ChangesetNumber);
+
 	return InCommand.bCommandSuccessful;
 }
 
 bool FPlasticUnlockWorker::UpdateStates()
 {
-	return PlasticSourceControlUtils::UpdateCachedStates(MoveTemp(States));
+	PlasticSourceControlUtils::UpdateCachedStates(MoveTemp(States));
+	// Note: Force the return to always trigger a refresh of the List of Locks
+	// It's needed since when removing the Lock of a Checked Out asset there is actually no change in local status as the asset remains Checked Out!
+	return true;
 }
 
 FName FPlasticGetBranchesWorker::GetName() const
@@ -1023,13 +1226,12 @@ bool FPlasticGetBranchesWorker::Execute(FPlasticSourceControlCommand& InCommand)
 	TSharedRef<FPlasticGetBranches, ESPMode::ThreadSafe> Operation = StaticCastSharedRef<FPlasticGetBranches>(InCommand.Operation);
 
 	{
-		const FDateTime& FromDate = Operation->FromDate;
-		InCommand.bCommandSuccessful = PlasticSourceControlUtils::RunGetBranches(FromDate, Operation->Branches, InCommand.ErrorMessages);
+		InCommand.bCommandSuccessful = PlasticSourceControlUtils::RunGetBranches(Operation->FromDate, Operation->Branches, InCommand.ErrorMessages);
 	}
 
 	{
 		FString RepositoryName, ServerUrl;
-		InCommand.bCommandSuccessful &= PlasticSourceControlUtils::GetWorkspaceInfo(CurrentBranchName, RepositoryName, ServerUrl, InCommand.ErrorMessages);
+		InCommand.bCommandSuccessful &= PlasticSourceControlUtils::GetWorkspaceInfo(InCommand.WorkspaceSelector, InCommand.BranchName, RepositoryName, ServerUrl, InCommand.ErrorMessages);
 	}
 
 	return InCommand.bCommandSuccessful;
@@ -1037,38 +1239,42 @@ bool FPlasticGetBranchesWorker::Execute(FPlasticSourceControlCommand& InCommand)
 
 bool FPlasticGetBranchesWorker::UpdateStates()
 {
-	GetProvider().SetBranchName(CurrentBranchName);
-
 	return false;
 }
 
 
-FName FPlasticSwitchToBranchWorker::GetName() const
+FName FPlasticSwitchWorker::GetName() const
 {
 	return "Switch";
 }
 
-bool FPlasticSwitchToBranchWorker::Execute(FPlasticSourceControlCommand& InCommand)
+bool FPlasticSwitchWorker::Execute(FPlasticSourceControlCommand& InCommand)
 {
-	TRACE_CPUPROFILER_EVENT_SCOPE(FPlasticSwitchToBranchWorker::Execute);
+	TRACE_CPUPROFILER_EVENT_SCOPE(FPlasticSwitchWorker::Execute);
 
 	check(InCommand.Operation->GetName() == GetName());
-	TSharedRef<FPlasticSwitchToBranch, ESPMode::ThreadSafe> Operation = StaticCastSharedRef<FPlasticSwitchToBranch>(InCommand.Operation);
+	TSharedRef<FPlasticSwitch, ESPMode::ThreadSafe> Operation = StaticCastSharedRef<FPlasticSwitch>(InCommand.Operation);
 
-	InCommand.bCommandSuccessful = PlasticSourceControlUtils::RunSwitchToBranch(Operation->BranchName, Operation->UpdatedFiles, InCommand.ErrorMessages);
+	InCommand.bCommandSuccessful = PlasticSourceControlUtils::RunSwitch(Operation->BranchName, Operation->ChangesetId, GetProvider().IsPartialWorkspace(), Operation->UpdatedFiles, InCommand.ErrorMessages);
+
+	// the current branch is used to asses the status of Retained Locks
+	{
+		FString RepositoryName, ServerUrl;
+		PlasticSourceControlUtils::GetWorkspaceInfo(InCommand.WorkspaceSelector, InCommand.BranchName, RepositoryName, ServerUrl, InCommand.ErrorMessages);
+		GetProvider().SetWorkspaceSelector(InCommand.WorkspaceSelector, InCommand.BranchName);
+		PlasticSourceControlUtils::InvalidateLocksCache();
+	}
 
 	// now update the status of the updated files
 	if (InCommand.bCommandSuccessful && Operation->UpdatedFiles.Num())
 	{
-		// the current branch is used to asses the status of Retained Locks
-		GetProvider().SetBranchName(Operation->BranchName);
-		InCommand.bCommandSuccessful = PlasticSourceControlUtils::RunUpdateStatus(Operation->UpdatedFiles, PlasticSourceControlUtils::EStatusSearchType::ControlledOnly, false, InCommand.ErrorMessages, States, InCommand.ChangesetNumber, InCommand.BranchName);
+		PlasticSourceControlUtils::RunUpdateStatus(Operation->UpdatedFiles, PlasticSourceControlUtils::EStatusSearchType::ControlledOnly, false, InCommand.ErrorMessages, States, InCommand.ChangesetNumber);
 	}
 
 	return InCommand.bCommandSuccessful;
 }
 
-bool FPlasticSwitchToBranchWorker::UpdateStates()
+bool FPlasticSwitchWorker::UpdateStates()
 {
 	return PlasticSourceControlUtils::UpdateCachedStates(MoveTemp(States));
 }
@@ -1091,7 +1297,8 @@ bool FPlasticMergeBranchWorker::Execute(FPlasticSourceControlCommand& InCommand)
 	// now update the status of the updated files
 	if (Operation->UpdatedFiles.Num())
 	{
-		InCommand.bCommandSuccessful = PlasticSourceControlUtils::RunUpdateStatus(Operation->UpdatedFiles, PlasticSourceControlUtils::EStatusSearchType::ControlledOnly, false, InCommand.ErrorMessages, States, InCommand.ChangesetNumber, InCommand.BranchName);
+		PlasticSourceControlUtils::InvalidateLocksCache();
+		PlasticSourceControlUtils::RunUpdateStatus(Operation->UpdatedFiles, PlasticSourceControlUtils::EStatusSearchType::ControlledOnly, false, InCommand.ErrorMessages, States, InCommand.ChangesetNumber);
 	}
 
 	return InCommand.bCommandSuccessful;
@@ -1169,6 +1376,70 @@ bool FPlasticDeleteBranchesWorker::UpdateStates()
 }
 
 
+FName FPlasticGetChangesetsWorker::GetName() const
+{
+	return "GetChangesets";
+}
+
+bool FPlasticGetChangesetsWorker::Execute(FPlasticSourceControlCommand& InCommand)
+{
+	TRACE_CPUPROFILER_EVENT_SCOPE(FPlasticGetChangesetsWorker::Execute);
+
+	check(InCommand.Operation->GetName() == GetName());
+	TSharedRef<FPlasticGetChangesets, ESPMode::ThreadSafe> Operation = StaticCastSharedRef<FPlasticGetChangesets>(InCommand.Operation);
+
+	{
+		InCommand.bCommandSuccessful = PlasticSourceControlUtils::RunGetChangesets(Operation->FromDate, Operation->Changesets, InCommand.ErrorMessages);
+	}
+
+	{
+		InCommand.bCommandSuccessful &= PlasticSourceControlUtils::GetChangesetNumber(InCommand.ChangesetNumber, InCommand.ErrorMessages);
+	}
+
+	return InCommand.bCommandSuccessful;
+}
+
+bool FPlasticGetChangesetsWorker::UpdateStates()
+{
+	return false;
+}
+
+
+
+FName FPlasticGetChangesetFilesWorker::GetName() const
+{
+	return "GetChangesetFiles";
+}
+
+bool FPlasticGetChangesetFilesWorker::Execute(FPlasticSourceControlCommand& InCommand)
+{
+	TRACE_CPUPROFILER_EVENT_SCOPE(FPlasticGetChangesetFilesWorker::Execute);
+
+	check(InCommand.Operation->GetName() == GetName());
+	TSharedRef<FPlasticGetChangesetFiles, ESPMode::ThreadSafe> Operation = StaticCastSharedRef<FPlasticGetChangesetFiles>(InCommand.Operation);
+
+	if (!Operation->Changeset.IsValid())
+	{
+		return false;
+	}
+
+	{
+		InCommand.bCommandSuccessful = PlasticSourceControlUtils::RunGetChangesetFiles(Operation->Changeset.ToSharedRef(), Operation->Files, InCommand.ErrorMessages);
+	}
+
+	{
+		InCommand.bCommandSuccessful &= PlasticSourceControlUtils::GetChangesetNumber(InCommand.ChangesetNumber, InCommand.ErrorMessages);
+	}
+
+	return InCommand.bCommandSuccessful;
+}
+
+bool FPlasticGetChangesetFilesWorker::UpdateStates()
+{
+	return false;
+}
+
+
 FName FPlasticUpdateStatusWorker::GetName() const
 {
 	return "UpdateStatus";
@@ -1189,14 +1460,14 @@ bool FPlasticUpdateStatusWorker::Execute(FPlasticSourceControlCommand& InCommand
 
 	if (Files.Num() > 0)
 	{
-		InCommand.bCommandSuccessful = PlasticSourceControlUtils::RunUpdateStatus(Files, PlasticSourceControlUtils::EStatusSearchType::All, Operation->ShouldUpdateHistory(), InCommand.ErrorMessages, States, InCommand.ChangesetNumber, InCommand.BranchName);
+		InCommand.bCommandSuccessful = PlasticSourceControlUtils::RunUpdateStatus(Files, PlasticSourceControlUtils::EStatusSearchType::All, Operation->ShouldUpdateHistory(), InCommand.ErrorMessages, States, InCommand.ChangesetNumber);
 		// Remove all "is not in a workspace" error and convert the result to "success" if there are no other errors
 		PlasticSourceControlUtils::RemoveRedundantErrors(InCommand, TEXT("is not in a workspace."));
 		if (!InCommand.bCommandSuccessful)
 		{
 			// In case of error, execute a 'checkconnection' command to check the connection to the server.
 			UE_LOG(LogSourceControl, Warning, TEXT("Error on 'status', execute a 'checkconnection' to test the connection to the server"));
-			InCommand.bConnectionDropped = !PlasticSourceControlUtils::RunCheckConnection(InCommand.BranchName, InCommand.RepositoryName, InCommand.ServerUrl, InCommand.InfoMessages, InCommand.ErrorMessages);
+			InCommand.bConnectionDropped = !PlasticSourceControlUtils::RunCheckConnection(InCommand.WorkspaceSelector, InCommand.BranchName, InCommand.RepositoryName, InCommand.ServerUrl, InCommand.InfoMessages, InCommand.ErrorMessages);
 			return false;
 		}
 
@@ -1223,7 +1494,7 @@ bool FPlasticUpdateStatusWorker::Execute(FPlasticSourceControlCommand& InCommand
 	{
 		TArray<FString> ProjectDirs;
 		ProjectDirs.Add(FPaths::ConvertRelativePathToFull(FPaths::ProjectContentDir()));
-		InCommand.bCommandSuccessful = PlasticSourceControlUtils::RunUpdateStatus(ProjectDirs, PlasticSourceControlUtils::EStatusSearchType::All, Operation->ShouldUpdateHistory(), InCommand.ErrorMessages, States, InCommand.ChangesetNumber, InCommand.BranchName);
+		InCommand.bCommandSuccessful = PlasticSourceControlUtils::RunUpdateStatus(ProjectDirs, PlasticSourceControlUtils::EStatusSearchType::All, Operation->ShouldUpdateHistory(), InCommand.ErrorMessages, States, InCommand.ChangesetNumber);
 	}
 	else
 	{
@@ -1246,7 +1517,7 @@ bool FPlasticUpdateStatusWorker::UpdateStates()
 	// Update affected changelists if any (in case of a file reverted outside of the Unreal Editor)
 	for (const FPlasticSourceControlState& NewState : States)
 	{
-		if (!NewState.IsCheckedOutImplementation())
+		if (!NewState.IsPendingChanges())
 		{
 			TSharedRef<FPlasticSourceControlState, ESPMode::ThreadSafe> State = GetProvider().GetStateInternal(NewState.GetFilename());
 			if (State->Changelist.IsInitialized())
@@ -1258,6 +1529,18 @@ bool FPlasticUpdateStatusWorker::UpdateStates()
 				State->Changelist.Reset();
 			}
 		}
+		else if (NewState.WorkspaceState == EWorkspaceState::Changed)
+		{
+			TSharedRef<FPlasticSourceControlState, ESPMode::ThreadSafe> State = GetProvider().GetStateInternal(NewState.GetFilename());
+			if (State->WorkspaceState != EWorkspaceState::Changed)
+			{
+				// 1- Add these files to the default changelist
+				TSharedRef<FPlasticSourceControlChangelistState, ESPMode::ThreadSafe> DefaultChangelist = GetProvider().GetStateInternal(FPlasticSourceControlChangelist::DefaultChangelist);
+				DefaultChangelist->Files.AddUnique(State);
+				// 2- And set the reference to the default changelist
+				State->Changelist = FPlasticSourceControlChangelist::DefaultChangelist;
+			}
+		}
 	}
 
 	return PlasticSourceControlUtils::UpdateCachedStates(MoveTemp(States));
@@ -1334,7 +1617,8 @@ bool FPlasticCopyWorker::Execute(FPlasticSourceControlCommand& InCommand)
 		TArray<FString> BothFiles;
 		BothFiles.Add(Origin);
 		BothFiles.Add(Destination);
-		PlasticSourceControlUtils::RunUpdateStatus(BothFiles, PlasticSourceControlUtils::EStatusSearchType::ControlledOnly, false, InCommand.ErrorMessages, States, InCommand.ChangesetNumber, InCommand.BranchName);
+		PlasticSourceControlUtils::InvalidateLocksCache();
+		PlasticSourceControlUtils::RunUpdateStatus(BothFiles, PlasticSourceControlUtils::EStatusSearchType::ControlledOnly, false, InCommand.ErrorMessages, States, InCommand.ChangesetNumber);
 	}
 	else
 	{
@@ -1360,19 +1644,28 @@ bool FPlasticSyncWorker::Execute(FPlasticSourceControlCommand& InCommand)
 {
 	TRACE_CPUPROFILER_EVENT_SCOPE(FPlasticSyncWorker::Execute);
 
+	TSharedRef<FPlasticSyncAll, ESPMode::ThreadSafe> SyncOperation = StaticCastSharedRef<FPlasticSyncAll>(InCommand.Operation);
+
 	TArray<FString> UpdatedFiles;
-	InCommand.bCommandSuccessful = PlasticSourceControlUtils::RunUpdate(InCommand.Files, GetProvider().IsPartialWorkspace(), UpdatedFiles, InCommand.ErrorMessages);
+	InCommand.bCommandSuccessful = PlasticSourceControlUtils::RunUpdate(InCommand.Files, GetProvider().IsPartialWorkspace(), SyncOperation->GetRevision(), UpdatedFiles, InCommand.ErrorMessages);
+
+	// the current branch is used to asses the status of Retained Locks
+	{
+		FString RepositoryName, ServerUrl;
+		PlasticSourceControlUtils::GetWorkspaceInfo(InCommand.WorkspaceSelector, InCommand.BranchName, RepositoryName, ServerUrl, InCommand.ErrorMessages);
+		GetProvider().SetWorkspaceSelector(InCommand.WorkspaceSelector, InCommand.BranchName);
+	}
 
 	// now update the status of the updated files
-	if (UpdatedFiles.Num())
+	if (InCommand.bCommandSuccessful && UpdatedFiles.Num())
 	{
-		PlasticSourceControlUtils::RunUpdateStatus(UpdatedFiles, PlasticSourceControlUtils::EStatusSearchType::ControlledOnly, false, InCommand.ErrorMessages, States, InCommand.ChangesetNumber, InCommand.BranchName);
+		PlasticSourceControlUtils::RunUpdateStatus(UpdatedFiles, PlasticSourceControlUtils::EStatusSearchType::ControlledOnly, false, InCommand.ErrorMessages, States, InCommand.ChangesetNumber);
 	}
 
 	if ((InCommand.Operation->GetName() == FName("SyncAll")))
 	{
-		TSharedRef<FPlasticSyncAll, ESPMode::ThreadSafe> Operation = StaticCastSharedRef<FPlasticSyncAll>(InCommand.Operation);
-		Operation->UpdatedFiles = MoveTemp(UpdatedFiles);
+		TSharedRef<FPlasticSyncAll, ESPMode::ThreadSafe> SyncAllOperation = StaticCastSharedRef<FPlasticSyncAll>(InCommand.Operation);
+		SyncAllOperation->UpdatedFiles = MoveTemp(UpdatedFiles);
 	}
 
 	return InCommand.bCommandSuccessful;
@@ -1419,7 +1712,8 @@ bool FPlasticResolveWorker::Execute(FPlasticSourceControlCommand& InCommand)
 	}
 
 	// now update the status of our files
-	PlasticSourceControlUtils::RunUpdateStatus(InCommand.Files, PlasticSourceControlUtils::EStatusSearchType::ControlledOnly, false, InCommand.ErrorMessages, States, InCommand.ChangesetNumber, InCommand.BranchName);
+	PlasticSourceControlUtils::InvalidateLocksCache();
+	PlasticSourceControlUtils::RunUpdateStatus(InCommand.Files, PlasticSourceControlUtils::EStatusSearchType::ControlledOnly, false, InCommand.ErrorMessages, States, InCommand.ChangesetNumber);
 
 	return InCommand.bCommandSuccessful;
 }
@@ -1495,6 +1789,11 @@ bool FPlasticGetPendingChangelistsWorker::UpdateStates()
 			for (const auto& FileState : OutCLFilesStates[StatusIndex])
 			{
 				TSharedRef<FPlasticSourceControlState, ESPMode::ThreadSafe> CachedFileState = GetProvider().GetStateInternal(FileState.LocalFilename);
+				if (CachedFileState->WorkspaceState == EWorkspaceState::Unknown)
+				{
+					CachedFileState->WorkspaceState = FileState.WorkspaceState;
+					CachedFileState->TimeStamp = Now;
+				}
 				CachedFileState->Changelist = CLStatus.Changelist;
 				ChangelistState->Files.AddUnique(CachedFileState);
 			}
@@ -1526,7 +1825,7 @@ bool FPlasticGetPendingChangelistsWorker::UpdateStates()
 	return bUpdated;
 }
 
-FPlasticSourceControlChangelist GenerateUniqueChangelistName(FPlasticSourceControlProvider& PlasticSourceControlProvider)
+static FPlasticSourceControlChangelist GenerateUniqueChangelistName(FPlasticSourceControlProvider& PlasticSourceControlProvider)
 {
 	FPlasticSourceControlChangelist NewChangelist;
 
@@ -1545,7 +1844,7 @@ FPlasticSourceControlChangelist GenerateUniqueChangelistName(FPlasticSourceContr
 	return NewChangelist;
 }
 
-FPlasticSourceControlChangelist CreatePendingChangelist(FPlasticSourceControlProvider& PlasticSourceControlProvider, const FString& InDescription, TArray<FString>& InInfoMessages, TArray<FString>& InErrorMessages)
+static FPlasticSourceControlChangelist CreatePendingChangelist(FPlasticSourceControlProvider& PlasticSourceControlProvider, const FString& InDescription, TArray<FString>& InInfoMessages, TArray<FString>& InErrorMessages)
 {
 	FPlasticSourceControlChangelist NewChangelist = GenerateUniqueChangelistName(PlasticSourceControlProvider);
 
@@ -1578,7 +1877,7 @@ FPlasticSourceControlChangelist CreatePendingChangelist(FPlasticSourceControlPro
 	return NewChangelist;
 }
 
-bool EditChangelistDescription(const FPlasticSourceControlProvider& PlasticSourceControlProvider, const FPlasticSourceControlChangelist& InChangelist, const FString& InDescription, TArray<FString>& InInfoMessages, TArray<FString>& InErrorMessages)
+static bool EditChangelistDescription(const FPlasticSourceControlProvider& PlasticSourceControlProvider, const FPlasticSourceControlChangelist& InChangelist, const FString& InDescription, TArray<FString>& InInfoMessages, TArray<FString>& InErrorMessages)
 {
 	TArray<FString> Parameters;
 	Parameters.Add(TEXT("edit"));
@@ -1601,12 +1900,12 @@ bool EditChangelistDescription(const FPlasticSourceControlProvider& PlasticSourc
 	}
 }
 
-bool MoveFilesToChangelist(const FPlasticSourceControlProvider& PlasticSourceControlProvider, const FPlasticSourceControlChangelist& InChangelist, const TArray<FString>& InFiles, TArray<FString>& OutResults, TArray<FString>& OutErrorMessages)
+static bool MoveFilesToChangelist(const FPlasticSourceControlProvider& InSourceControlProvider, const FPlasticSourceControlChangelist& InChangelist, const TArray<FString>& InFiles, TArray<FString>& OutResults, TArray<FString>& OutErrorMessages)
 {
 	if (InFiles.Num() > 0)
 	{
 		TArray<FString> Parameters;
-		if (PlasticSourceControlProvider.GetPlasticScmVersion() < PlasticSourceControlVersions::NewChangelistFileArgs)
+		if (InSourceControlProvider.GetPlasticScmVersion() < PlasticSourceControlVersions::NewChangelistFileArgs)
 		{
 			Parameters.Add(TEXT("\"") + InChangelist.GetName() + TEXT("\""));
 			Parameters.Add(TEXT("add"));
@@ -1617,13 +1916,28 @@ bool MoveFilesToChangelist(const FPlasticSourceControlProvider& PlasticSourceCon
 			const FScopedTempFile ChangelistNameFile(InChangelist.GetName());
 			Parameters.Add(FString::Printf(TEXT("--namefile=\"%s\""), *ChangelistNameFile.GetFilename()));
 			Parameters.Add(TEXT("add"));
-			UE_LOG(LogSourceControl, Verbose, TEXT("MoveFilesToChangelist(%s)"), *InChangelist.GetName());
+			UE_LOG(LogSourceControl, Verbose, TEXT("MoveFilesToChangelist: %s"), *InChangelist.GetName());
 			return PlasticSourceControlUtils::RunCommand(TEXT("changelist"), Parameters, InFiles, OutResults, OutErrorMessages);
 		}
 	}
 	return true;
 }
 
+// Old "cm" doesn't support newlines, quotes, and question marks on changelist's name or description
+static FString CleanupChangelistDescription(const FPlasticSourceControlProvider& InSourceControlProvider, const FText& InDescription)
+{
+	FString Description = InDescription.ToString();
+	if (InSourceControlProvider.GetPlasticScmVersion() < PlasticSourceControlVersions::NewChangelistFileArgs)
+	{
+		Description.ReplaceInline(TEXT("\r\n"), TEXT(" "), ESearchCase::CaseSensitive);
+		Description.ReplaceCharInline(TEXT('\n'), TEXT(' '), ESearchCase::CaseSensitive);
+		Description.ReplaceCharInline(TEXT('\"'), TEXT('\''), ESearchCase::CaseSensitive);
+		Description.ReplaceCharInline(TEXT('?'), TEXT('.'), ESearchCase::CaseSensitive);
+		Description.ReplaceCharInline(TEXT('*'), TEXT('.'), ESearchCase::CaseSensitive);
+	}
+	return Description;
+}
+
 FPlasticNewChangelistWorker::FPlasticNewChangelistWorker(FPlasticSourceControlProvider& InSourceControlProvider)
 	: IPlasticSourceControlWorker(InSourceControlProvider)
 	, NewChangelistState(NewChangelist)
@@ -1641,17 +1955,7 @@ bool FPlasticNewChangelistWorker::Execute(class FPlasticSourceControlCommand& In
 
 	check(InCommand.Operation->GetName() == GetName());
 	TSharedRef<FNewChangelist, ESPMode::ThreadSafe> Operation = StaticCastSharedRef<FNewChangelist>(InCommand.Operation);
-
-	FString Description = Operation->GetDescription().ToString();
-	// Note: old "cm" doesn't support newlines, quotes, and question marks on changelist's name or description
-	if (GetProvider().GetPlasticScmVersion() < PlasticSourceControlVersions::NewChangelistFileArgs)
-	{
-		Description.ReplaceInline(TEXT("\r\n"), TEXT(" "), ESearchCase::CaseSensitive);
-		Description.ReplaceCharInline(TEXT('\n'), TEXT(' '), ESearchCase::CaseSensitive);
-		Description.ReplaceCharInline(TEXT('\"'), TEXT('\''), ESearchCase::CaseSensitive);
-		Description.ReplaceCharInline(TEXT('?'), TEXT('.'), ESearchCase::CaseSensitive);
-		Description.ReplaceCharInline(TEXT('*'), TEXT('.'), ESearchCase::CaseSensitive);
-	}
+	FString Description = CleanupChangelistDescription(GetProvider(), Operation->GetDescription());
 
 	// Create a new numbered persistent changelist ala Perforce
 	NewChangelist = CreatePendingChangelist(GetProvider(), Description, InCommand.InfoMessages, InCommand.ErrorMessages);
@@ -1668,10 +1972,19 @@ bool FPlasticNewChangelistWorker::Execute(class FPlasticSourceControlCommand& In
 
 		if (InCommand.Files.Num() > 0)
 		{
-			InCommand.bCommandSuccessful = MoveFilesToChangelist(GetProvider(), NewChangelist, InCommand.Files, InCommand.InfoMessages, InCommand.ErrorMessages);
+			// Filter out files that cannot be moved from the Default changelist; files with Pending Changes but not actually Checked Out
+			TArray<FString> FilesToMove = GetFilesCheckedOut(GetProvider(), InCommand.Files);
+			if (FilesToMove.Num() == 0)
+			{
+				UE_LOG(LogSourceControl, Error, TEXT("No file can be moved."));
+				InCommand.bCommandSuccessful = false;
+				return InCommand.bCommandSuccessful;
+			}
+
+			InCommand.bCommandSuccessful = MoveFilesToChangelist(GetProvider(), NewChangelist, FilesToMove, InCommand.InfoMessages, InCommand.ErrorMessages);
 			if (InCommand.bCommandSuccessful)
 			{
-				MovedFiles = InCommand.Files;
+				MovedFiles = MoveTemp(FilesToMove);
 			}
 		}
 	}
@@ -1781,16 +2094,7 @@ bool FPlasticEditChangelistWorker::Execute(class FPlasticSourceControlCommand& I
 	check(InCommand.Operation->GetName() == GetName());
 	TSharedRef<FEditChangelist, ESPMode::ThreadSafe> Operation = StaticCastSharedRef<FEditChangelist>(InCommand.Operation);
 
-	EditedDescription = Operation->GetDescription().ToString();
-	// Note: old "cm" doesn't support newlines, quotes, and question marks on changelist's name or description
-	if (GetProvider().GetPlasticScmVersion() < PlasticSourceControlVersions::NewChangelistFileArgs)
-	{
-		EditedDescription.ReplaceInline(TEXT("\r\n"), TEXT(" "), ESearchCase::CaseSensitive);
-		EditedDescription.ReplaceCharInline(TEXT('\n'), TEXT(' '), ESearchCase::CaseSensitive);
-		EditedDescription.ReplaceCharInline(TEXT('\"'), TEXT('\''), ESearchCase::CaseSensitive);
-		EditedDescription.ReplaceCharInline(TEXT('?'), TEXT('.'), ESearchCase::CaseSensitive);
-		EditedDescription.ReplaceCharInline(TEXT('*'), TEXT('.'), ESearchCase::CaseSensitive);
-	}
+	EditedDescription = CleanupChangelistDescription(GetProvider(), Operation->GetDescription());
 
 	if (InCommand.Changelist.IsDefault())
 	{
@@ -1800,8 +2104,20 @@ bool FPlasticEditChangelistWorker::Execute(class FPlasticSourceControlCommand& I
 		{
 			// And then move all its files to the new changelist
 			TSharedRef<FPlasticSourceControlChangelistState, ESPMode::ThreadSafe> ChangelistState = GetProvider().GetStateInternal(InCommand.Changelist);
-			ReopenedFiles = FileNamesFromFileStates(ChangelistState->Files);
-			InCommand.bCommandSuccessful = MoveFilesToChangelist(GetProvider(), EditedChangelist, ReopenedFiles, InCommand.InfoMessages, InCommand.ErrorMessages);
+			// Filter out files that cannot be moved from the Default changelist; files with Pending Changes but not actually Checked Out
+			TArray<FString> FilesToMove = GetFilesCheckedOut(GetProvider(), FileNamesFromFileStates(ChangelistState->Files));
+			if (FilesToMove.Num() == 0)
+			{
+				UE_LOG(LogSourceControl, Error, TEXT("No file can be moved."));
+				InCommand.bCommandSuccessful = false;
+				return InCommand.bCommandSuccessful;
+			}
+
+			InCommand.bCommandSuccessful = MoveFilesToChangelist(GetProvider(), EditedChangelist, FilesToMove, InCommand.InfoMessages, InCommand.ErrorMessages);
+			if (InCommand.bCommandSuccessful)
+			{
+				ReopenedFiles = MoveTemp(FilesToMove);
+			}
 		}
 	}
 	else
@@ -1865,10 +2181,19 @@ bool FPlasticReopenWorker::Execute(FPlasticSourceControlCommand& InCommand)
 
 	check(InCommand.Operation->GetName() == GetName());
 
-	InCommand.bCommandSuccessful = MoveFilesToChangelist(GetProvider(), InCommand.Changelist, InCommand.Files, InCommand.InfoMessages, InCommand.ErrorMessages);
+	// Filter out files that cannot be moved from the Default changelist; files with Pending Changes but not actually Checked Out
+	TArray<FString> FilesToMove = GetFilesCheckedOut(GetProvider(), InCommand.Files);
+	if (FilesToMove.Num() == 0)
+	{
+		UE_LOG(LogSourceControl, Error, TEXT("No file can be moved."));
+		InCommand.bCommandSuccessful = false;
+		return InCommand.bCommandSuccessful;
+	}
+
+	InCommand.bCommandSuccessful = MoveFilesToChangelist(GetProvider(), InCommand.Changelist, FilesToMove, InCommand.InfoMessages, InCommand.ErrorMessages);
 	if (InCommand.bCommandSuccessful)
 	{
-		ReopenedFiles = InCommand.Files;
+		ReopenedFiles = MoveTemp(FilesToMove);
 		DestinationChangelist = InCommand.Changelist;
 	}
 
@@ -1909,7 +2234,7 @@ bool FPlasticReopenWorker::UpdateStates()
 	}
 }
 
-bool CreateShelve(const FString& InChangelistName, const FString& InChangelistDescription, const TArray<FString>& InFilesToShelve, int32& OutShelveId, TArray<FString>& OutErrorMessages)
+static bool CreateShelve(const FString& InChangelistName, const FString& InChangelistDescription, const TArray<FString>& InFilesToShelve, int32& OutShelveId, TArray<FString>& OutErrorMessages)
 {
 	TArray<FString> Results;
 	TArray<FString> Parameters;
@@ -1938,7 +2263,7 @@ bool CreateShelve(const FString& InChangelistName, const FString& InChangelistDe
 	return OutShelveId != ISourceControlState::INVALID_REVISION;
 }
 
-bool DeleteShelve(const int32 InShelveId, TArray<FString>& OutErrorMessages)
+static bool DeleteShelve(const int32 InShelveId, TArray<FString>& OutErrorMessages)
 {
 	TArray<FString> Results;
 	TArray<FString> Parameters;
@@ -2197,6 +2522,7 @@ bool FPlasticUnshelveWorker::Execute(FPlasticSourceControlCommand& InCommand)
 			TArray<FString> Parameters;
 			Parameters.Add(TEXT("apply"));
 			Parameters.Add(FString::Printf(TEXT("sh:%d"), ChangelistState->ShelveId));
+			// TODO --noinput
 			InCommand.bCommandSuccessful = PlasticSourceControlUtils::RunCommand(TEXT("shelveset"), Parameters, FilesToUnshelve, InCommand.InfoMessages, InCommand.ErrorMessages);
 		}
 
@@ -2213,7 +2539,8 @@ bool FPlasticUnshelveWorker::Execute(FPlasticSourceControlCommand& InCommand)
 	if (InCommand.bCommandSuccessful)
 	{
 		// now update the status of our files
-		PlasticSourceControlUtils::RunUpdateStatus(InCommand.Files, PlasticSourceControlUtils::EStatusSearchType::ControlledOnly, false, InCommand.ErrorMessages, States, InCommand.ChangesetNumber, InCommand.BranchName);
+		PlasticSourceControlUtils::InvalidateLocksCache();
+		PlasticSourceControlUtils::RunUpdateStatus(InCommand.Files, PlasticSourceControlUtils::EStatusSearchType::ControlledOnly, false, InCommand.ErrorMessages, States, InCommand.ChangesetNumber);
 
 		ChangelistToUpdate = InCommand.Changelist;
 	}
@@ -2421,22 +2748,22 @@ bool FPlasticGetFileWorker::Execute(FPlasticSourceControlCommand& InCommand)
 
 	TSharedRef<FGetFile, ESPMode::ThreadSafe> Operation = StaticCastSharedRef<FGetFile>(InCommand.Operation);
 
-	const TSharedRef<FPlasticSourceControlRevision, ESPMode::ThreadSafe> SourceControlRevision = MakeShared<FPlasticSourceControlRevision>();
-	SourceControlRevision->Filename = FPaths::ConvertRelativePathToFull(Operation->GetDepotFilePath());
+	FPlasticSourceControlRevision SourceControlRevision;
+	SourceControlRevision.Filename = FPaths::ConvertRelativePathToFull(Operation->GetDepotFilePath());
 
 	if (Operation->IsShelve())
 	{
-		SourceControlRevision->ShelveId = FCString::Atoi(*Operation->GetChangelistNumber());
-		UE_LOG(LogSourceControl, Log, TEXT("GetFile(ShelveId:%d)"), SourceControlRevision->ShelveId);
+		SourceControlRevision.ShelveId = FCString::Atoi(*Operation->GetChangelistNumber());
+		UE_LOG(LogSourceControl, Log, TEXT("GetFile(ShelveId:%d)"), SourceControlRevision.ShelveId);
 	}
 	else
 	{
-		SourceControlRevision->RevisionId = FCString::Atoi(*Operation->GetRevisionNumber());
-		UE_LOG(LogSourceControl, Log, TEXT("GetFile(revid:%d)"), SourceControlRevision->RevisionId);
+		SourceControlRevision.RevisionId = FCString::Atoi(*Operation->GetRevisionNumber());
+		UE_LOG(LogSourceControl, Log, TEXT("GetFile(revid:%d)"), SourceControlRevision.RevisionId);
 	}
 
 	FString OutFilename;
-	InCommand.bCommandSuccessful = SourceControlRevision->Get(OutFilename, InCommand.Concurrency);
+	InCommand.bCommandSuccessful = SourceControlRevision.Get(OutFilename, InCommand.Concurrency);
 	if (InCommand.bCommandSuccessful)
 	{
 		Operation->SetOutPackageFilename(OutFilename);
