PR #12084: [Source Control] Update Plastic SCM provider to 1.11.0 (from 1.9.0)

diff --git a/Engine/Plugins/Developer/PlasticSourceControl/Source/PlasticSourceControl/Private/PlasticSourceControlOperations.h b/Engine/Plugins/Developer/PlasticSourceControl/Source/PlasticSourceControl/Private/PlasticSourceControlOperations.h
--- a/Engine/Plugins/Developer/PlasticSourceControl/Source/PlasticSourceControl/Private/PlasticSourceControlOperations.h
+++ b/Engine/Plugins/Developer/PlasticSourceControl/Source/PlasticSourceControl/Private/PlasticSourceControlOperations.h
@@ -16,12 +16,16 @@
 
 class FPlasticSourceControlProvider;
 typedef TSharedRef<class FPlasticSourceControlBranch, ESPMode::ThreadSafe> FPlasticSourceControlBranchRef;
+typedef TSharedRef<class FPlasticSourceControlChangeset, ESPMode::ThreadSafe> FPlasticSourceControlChangesetRef;
+typedef TSharedPtr<class FPlasticSourceControlChangeset, ESPMode::ThreadSafe> FPlasticSourceControlChangesetPtr;
+typedef TSharedRef<class FPlasticSourceControlLock, ESPMode::ThreadSafe> FPlasticSourceControlLockRef;
 
 
 /**
  * Internal operation used to revert checked-out unchanged files
+ *
+ * NOTE: added to Engine in Unreal Engine 5 for changelists
 */
-// NOTE: added to Engine in Unreal Engine 5 for changelists
 class FPlasticRevertUnchanged final : public FSourceControlOperationBase
 {
 public:
@@ -34,6 +38,8 @@ public:
 
 /**
  * Internal operation used to sync all files in the workspace
+ *
+ * NOTE: override the standard "Sync" operation to provide a list of updated files and a custom progress string
 */
 class FPlasticSyncAll final : public FSync
 {
@@ -41,6 +47,8 @@ public:
 	// ISourceControlOperation interface
 	virtual FName GetName() const override;
 
+	virtual FText GetInProgressString() const override;
+
 	/** List of files updated by the operation */
 	TArray<FString> UpdatedFiles;
 };
@@ -62,6 +70,25 @@ public:
 };
 
 
+/**
+ * Internal operation used to revert files to a previous revision
+*/
+class FPlasticRevertToRevision final : public FSourceControlOperationBase
+{
+public:
+	// ISourceControlOperation interface
+	virtual FName GetName() const override;
+
+	virtual FText GetInProgressString() const override;
+
+	/** List of files updated by the operation */
+	TArray<FString> UpdatedFiles;
+
+	// Changeset to revert the files to
+	int32 ChangesetId = ISourceControlState::INVALID_REVISION;
+};
+
+
 /**
 * Internal operation used to create a new Workspace and a new Repository
 */
@@ -93,6 +120,22 @@ public:
 };
 
 
+/**
+ * Internal operation to list locks, aka "cm lock list"
+*/
+class FPlasticGetLocks final : public FSourceControlOperationBase
+{
+public:
+	// ISourceControlOperation interface
+	virtual FName GetName() const override;
+
+	virtual FText GetInProgressString() const override;
+
+	// List of locks found
+	TArray<FPlasticSourceControlLockRef> Locks;
+};
+
+
 /**
  * Internal operation used to release or remove Lock(s) on file(s)
 */
@@ -104,6 +147,9 @@ public:
 
 	virtual FText GetInProgressString() const override;
 
+	// Locks to unlock, including the Item Id and branch name
+	TArray<FPlasticSourceControlLockRef> Locks;
+
 	// Release the Lock(s), and optionally remove (delete) them completely
 	bool bRemove = false;
 };
@@ -120,7 +166,7 @@ public:
 
 	virtual FText GetInProgressString() const override;
 
-	// Limit the list of branches to ones created from this date
+	// Limit the list of branches to ones created from this date (optional, filtering enabled by default)
 	FDateTime FromDate;
 
 	// List of branches found
@@ -129,9 +175,9 @@ public:
 
 
 /**
- * Internal operation used to switch the workspace to another branch
+ * Internal operation used to switch the workspace to a branch or a changeset
 */
-class FPlasticSwitchToBranch final : public FSourceControlOperationBase
+class FPlasticSwitch final : public FSourceControlOperationBase
 {
 public:
 	// ISourceControlOperation interface
@@ -139,9 +185,12 @@ public:
 
 	virtual FText GetInProgressString() const override;
 
-	// Branch to switch the workspace to
+	// Branch to switch the workspace to (optional, only apply if ChangesetId is not set)
 	FString BranchName;
 
+	// Changeset to switch the workspace to (optional, overrides the BranchName if set)
+	int32 ChangesetId = ISourceControlState::INVALID_REVISION;
+
 	/** List of files updated by the operation */
 	TArray<FString> UpdatedFiles;
 };
@@ -213,6 +262,44 @@ public:
 };
 
 
+/**
+ * Internal operation to list changesets, aka "cm find changesets"
+*/
+class FPlasticGetChangesets final : public FSourceControlOperationBase
+{
+public:
+	// ISourceControlOperation interface
+	virtual FName GetName() const override;
+
+	virtual FText GetInProgressString() const override;
+
+	// Limit the list of changesets to ones created from this date (optional, filtering enabled by default)
+	FDateTime FromDate;
+
+	// List of changesets found
+	TArray<FPlasticSourceControlChangesetRef> Changesets;
+};
+
+
+/**
+ * Internal operation to list files in a changeset, using "cm log cs:<ChangesetId>"
+*/
+class FPlasticGetChangesetFiles final : public FSourceControlOperationBase
+{
+public:
+	// ISourceControlOperation interface
+	virtual FName GetName() const override;
+
+	virtual FText GetInProgressString() const override;
+
+	// Changeset to list files from
+	FPlasticSourceControlChangesetPtr Changeset;
+
+	// List of files changed in the changeset
+	TArray<FPlasticSourceControlStateRef> Files;
+};
+
+
 /** Called when first activated on a project, and then at project load time.
  *  Look for the root directory of the Plastic workspace (where the ".plastic/" subdirectory is located). */
 class FPlasticConnectWorker final : public IPlasticSourceControlWorker
@@ -372,6 +459,24 @@ public:
 	TArray<FPlasticSourceControlState> States;
 };
 
+/** Revert file(s) to selected revision. */
+class FPlasticRevertToRevisionWorker final : public IPlasticSourceControlWorker
+{
+public:
+	explicit FPlasticRevertToRevisionWorker(FPlasticSourceControlProvider& InSourceControlProvider)
+		: IPlasticSourceControlWorker(InSourceControlProvider)
+	{}
+	virtual ~FPlasticRevertToRevisionWorker() = default;
+	// IPlasticSourceControlWorker interface
+	virtual FName GetName() const override;
+	virtual bool Execute(class FPlasticSourceControlCommand& InCommand) override;
+	virtual bool UpdateStates() override;
+
+public:
+	/** Temporary states for results */
+	TArray<FPlasticSourceControlState> States;
+};
+
 /** Create a new Workspace and a new Repository */
 class FPlasticMakeWorkspaceWorker final : public IPlasticSourceControlWorker
 {
@@ -404,6 +509,20 @@ public:
 	TArray<FPlasticSourceControlState> States;
 };
 
+/** list locks. */
+class FPlasticGetLocksWorker final : public IPlasticSourceControlWorker
+{
+public:
+	explicit FPlasticGetLocksWorker(FPlasticSourceControlProvider& InSourceControlProvider)
+		: IPlasticSourceControlWorker(InSourceControlProvider)
+	{}
+	virtual ~FPlasticGetLocksWorker() = default;
+	// IPlasticSourceControlWorker interface
+	virtual FName GetName() const override;
+	virtual bool Execute(class FPlasticSourceControlCommand& InCommand) override;
+	virtual bool UpdateStates() override;
+};
+
 /** release or remove Lock(s) on file(s). */
 class FPlasticUnlockWorker final : public IPlasticSourceControlWorker
 {
@@ -434,19 +553,16 @@ public:
 	virtual FName GetName() const override;
 	virtual bool Execute(class FPlasticSourceControlCommand& InCommand) override;
 	virtual bool UpdateStates() override;
-
-	// Current branch the workspace is on (at the end of the operation)
-	FString CurrentBranchName;
 };
 
 /** Switch workspace to another branch. */
-class FPlasticSwitchToBranchWorker final : public IPlasticSourceControlWorker
+class FPlasticSwitchWorker final : public IPlasticSourceControlWorker
 {
 public:
-	explicit FPlasticSwitchToBranchWorker(FPlasticSourceControlProvider& InSourceControlProvider)
+	explicit FPlasticSwitchWorker(FPlasticSourceControlProvider& InSourceControlProvider)
 		: IPlasticSourceControlWorker(InSourceControlProvider)
 	{}
-	virtual ~FPlasticSwitchToBranchWorker() = default;
+	virtual ~FPlasticSwitchWorker() = default;
 	// IPlasticSourceControlWorker interface
 	virtual FName GetName() const override;
 	virtual bool Execute(class FPlasticSourceControlCommand& InCommand) override;
@@ -517,6 +633,37 @@ public:
 	virtual bool UpdateStates() override;
 };
 
+/** list changesets. */
+class FPlasticGetChangesetsWorker final : public IPlasticSourceControlWorker
+{
+public:
+	explicit FPlasticGetChangesetsWorker(FPlasticSourceControlProvider& InSourceControlProvider)
+		: IPlasticSourceControlWorker(InSourceControlProvider)
+	{}
+	virtual ~FPlasticGetChangesetsWorker() = default;
+	// IPlasticSourceControlWorker interface
+	virtual FName GetName() const override;
+	virtual bool Execute(class FPlasticSourceControlCommand& InCommand) override;
+	virtual bool UpdateStates() override;
+
+	// Current changeset the workspace is on (at the end of the operation)
+	int32 CurrentChangesetId;
+};
+
+/** list files in changeset. */
+class FPlasticGetChangesetFilesWorker final : public IPlasticSourceControlWorker
+{
+public:
+	explicit FPlasticGetChangesetFilesWorker(FPlasticSourceControlProvider& InSourceControlProvider)
+		: IPlasticSourceControlWorker(InSourceControlProvider)
+	{}
+	virtual ~FPlasticGetChangesetFilesWorker() = default;
+	// IPlasticSourceControlWorker interface
+	virtual FName GetName() const override;
+	virtual bool Execute(class FPlasticSourceControlCommand& InCommand) override;
+	virtual bool UpdateStates() override;
+};
+
 /** Plastic update the workspace to latest changes */
 class FPlasticSyncWorker final : public IPlasticSourceControlWorker
 {
diff --git a/Engine/Plugins/Developer/PlasticSourceControl/Source/PlasticSourceControl/Private/PlasticSourceControlParsers.cpp b/Engine/Plugins/Developer/PlasticSourceControl/Source/PlasticSourceControl/Private/PlasticSourceControlParsers.cpp
--- a/Engine/Plugins/Developer/PlasticSourceControl/Source/PlasticSourceControl/Private/PlasticSourceControlParsers.cpp
+++ b/Engine/Plugins/Developer/PlasticSourceControl/Source/PlasticSourceControl/Private/PlasticSourceControlParsers.cpp
@@ -3,6 +3,8 @@
 #include "PlasticSourceControlParsers.h"
 
 #include "PlasticSourceControlBranch.h"
+#include "PlasticSourceControlChangeset.h"
+#include "PlasticSourceControlLock.h"
 #include "PlasticSourceControlModule.h"
 #include "PlasticSourceControlProvider.h"
 #include "PlasticSourceControlState.h"
@@ -23,25 +25,57 @@ namespace PlasticSourceControlParsers
 #define FILE_STATUS_SEPARATOR TEXT(";")
 
 
-bool ParseWorkspaceInfo(TArray<FString>& InResults, FString& OutBranchName, FString& OutRepositoryName, FString& OutServerUrl)
+/**
+ * Parse the output of the "cm profile list --format="{server};{user}" command
+ *
+ * Example:
+localhost:8087|sebastien.rombauts
[REDACTED_EMAIL]
+SRombautsU@[REDACTED_EMAIL]
+*/
+bool ParseProfileInfo(TArray<FString>& InResults, const FString& InServerUrl, FString& OutUserName)
+{
+	for (const FString& Result : InResults)
+	{
+		TArray<FString> ProfileInfos;
+		Result.ParseIntoArray(ProfileInfos, FILE_STATUS_SEPARATOR, false); // Don't cull empty values
+		if (ProfileInfos.Num() == 2)
+		{
+			if (ProfileInfos[0] == InServerUrl)
+			{
+				OutUserName = ProfileInfos[1];
+				return true;
+			}
+		}
+	}
+
+	return false;
+}
+
+/**
+ * Parse  workspace information, in the form "Branch /main@UE5PlasticPluginDev@localhost:8087"
+ *                                        or "Branch /main@UE5PlasticPluginDev@test@cloud" (when connected to the cloud)
+ *                                        or "Branch /main@rep:UE5OpenWorldPerfTest@repserver:test@cloud"
+ *                                        or "Changeset 1234@UE5PlasticPluginDev@test@cloud" (when the workspace is switched to a changeset instead of a branch)
+ *                                        or "Label 1.10.0@UE5PlasticPluginDev@test@cloud" (when the workspace is switched to a label instead of a branch)
+*/
+bool ParseWorkspaceInfo(TArray<FString>& InResults, FString& OutWorkspaceSelector, FString& OutBranchName, FString& OutRepositoryName, FString& OutServerUrl)
 {
 	if (InResults.Num() == 0)
 	{
 		return false;
 	}
 
-	// Get workspace information, in the form "Branch /main@UE5PlasticPluginDev@localhost:8087"
-	//                                     or "Branch /main@UE5PlasticPluginDev@test@cloud" (when connected to the cloud)
-	//                                     or "Branch /main@rep:UE5OpenWorldPerfTest@repserver:test@cloud"
-	//                                     or "Changeset 1234@UE5PlasticPluginDev@test@cloud" (when the workspace is switched on a changeset instead of a branch)
 	static const FString BranchPrefix(TEXT("Branch "));
 	static const FString ChangesetPrefix(TEXT("Changeset "));
 	static const FString LabelPrefix(TEXT("Label "));
 	static const FString RepPrefix(TEXT("rep:"));
 	static const FString RepserverPrefix(TEXT("repserver:"));
+	bool bIsABranch = false;
 	FString& WorkspaceInfo = InResults[0];
 	if (WorkspaceInfo.StartsWith(BranchPrefix, ESearchCase::CaseSensitive))
 	{
+		bIsABranch = true;
 		WorkspaceInfo.RightChopInline(BranchPrefix.Len());
 	}
 	else if (WorkspaceInfo.StartsWith(ChangesetPrefix, ESearchCase::CaseSensitive))
@@ -61,7 +95,11 @@ bool ParseWorkspaceInfo(TArray<FString>& InResults, FString& OutBranchName, FStr
 	WorkspaceInfo.ParseIntoArray(WorkspaceInfos, TEXT("@"), false); // Don't cull empty values
 	if (WorkspaceInfos.Num() >= 3)
 	{
-		OutBranchName = MoveTemp(WorkspaceInfos[0]);
+		OutWorkspaceSelector = MoveTemp(WorkspaceInfos[0]);
+		if (bIsABranch)
+		{
+			OutBranchName = OutWorkspaceSelector;
+		}
 		OutRepositoryName = MoveTemp(WorkspaceInfos[1]);
 		OutServerUrl = MoveTemp(WorkspaceInfos[2]);
 
@@ -90,7 +128,7 @@ bool ParseWorkspaceInfo(TArray<FString>& InResults, FString& OutBranchName, FStr
 }
 
 /**
-* Parse the current changeset from the header returned by "cm status --machinereadable --header --fieldseparator=;"
+* Parse the current changeset from the header returned by "cm status --machinereadable --fieldseparator=;"
 *
 * Get workspace status in one of the form
 STATUS;41;UEPlasticPluginDev;localhost:8087
@@ -129,11 +167,11 @@ static EWorkspaceState StateFromStatus(const FString& InFileStatus, const bool b
 {
 	EWorkspaceState State;
 
-	if (InFileStatus == "CH") // Modified but not Checked-Out
+	if (InFileStatus == TEXT("CH")) // Modified but not Checked-Out
 	{
 		State = EWorkspaceState::Changed;
 	}
-	else if (InFileStatus == "CO") // Checked-Out with no change, or "don't know" if using on an old version of cm
+	else if (InFileStatus == TEXT("CO")) // Checked-Out with no change, or "don't know" if using on an old version of cm
 	{
 		// Recent version can distinguish between CheckedOut with or with no changes
 		if (bInUsesCheckedOutChanged)
@@ -145,7 +183,7 @@ static EWorkspaceState StateFromStatus(const FString& InFileStatus, const bool b
 			State = EWorkspaceState::CheckedOutChanged; // Older version; need to assume it is changed to retain behavior
 		}
 	}
-	else if (InFileStatus == "CO+CH") // Checked-Out and changed from the new --iscochanged
+	else if (InFileStatus == TEXT("CO+CH")) // Checked-Out and changed from the new --iscochanged
 	{
 		State = EWorkspaceState::CheckedOutChanged; // Recent version; here it's checkedout with changes
 	}
@@ -161,23 +199,23 @@ static EWorkspaceState StateFromStatus(const FString& InFileStatus, const bool b
 	{
 		State = EWorkspaceState::Replaced;
 	}
-	else if (InFileStatus == "AD")
+	else if (InFileStatus == TEXT("AD"))
 	{
 		State = EWorkspaceState::Added;
 	}
-	else if ((InFileStatus == "PR") || (InFileStatus == "LM")) // Not Controlled/Not in Depot/Untracked (or Locally Moved/Renamed)
+	else if ((InFileStatus == TEXT("PR")) || (InFileStatus == TEXT("LM"))) // Not Controlled/Not in Depot/Untracked (or Locally Moved/Renamed)
 	{
 		State = EWorkspaceState::Private;
 	}
-	else if (InFileStatus == "IG")
+	else if (InFileStatus == TEXT("IG"))
 	{
 		State = EWorkspaceState::Ignored;
 	}
-	else if (InFileStatus == "DE")
+	else if (InFileStatus == TEXT("DE"))
 	{
 		State = EWorkspaceState::Deleted; // Deleted (removed from source control)
 	}
-	else if (InFileStatus == "LD")
+	else if (InFileStatus.Contains(TEXT("LD"))) // "LD", "AD+LD"
 	{
 		State = EWorkspaceState::LocallyDeleted; // Locally Deleted (ie. missing)
 	}
@@ -234,7 +272,7 @@ static FPlasticSourceControlState StateFromStatusResult(const FString& InResult,
  * @brief Parse status results in case of a regular operation for a list of files (not for a whole directory).
  *
  * This is the most common scenario, for any operation from the Content Browser or the View Changes window.
- * 
+ *
  * In this case, iterates on the list of files the Editor provides,
  * searching corresponding file status from the array of strings results of a "status" command.
  *
@@ -410,20 +448,25 @@ public:
 	TArray<FString> Files;
 };
 
-FSmartLockInfoParser::FSmartLockInfoParser(const FString& InResult)
+FPlasticSourceControlLock ParseLockInfo(const FString& InResult)
 {
+	FPlasticSourceControlLock Lock;
 	TArray<FString> SmartLockInfos;
 	const int32 NbElmts = InResult.ParseIntoArray(SmartLockInfos, FILE_STATUS_SEPARATOR, false);
 	if (NbElmts >= 12)
 	{
-		Repository = MoveTemp(SmartLockInfos[0]);
-		ItemId = FCString::Atoi(*SmartLockInfos[1]);
-		FDateTime::ParseIso8601(*SmartLockInfos[3], Date);
-		BranchName = MoveTemp(SmartLockInfos[6]);
-		Status = MoveTemp(SmartLockInfos[8]);
-		Owner = PlasticSourceControlUtils::UserNameToDisplayName(MoveTemp(SmartLockInfos[9]));
-		Filename = MoveTemp(SmartLockInfos[11]);
-	}
+		Lock.ItemId = FCString::Atoi(*SmartLockInfos[1]);
+		FDateTime::ParseIso8601(*SmartLockInfos[3], Lock.Date);
+		Lock.DestinationBranch = MoveTemp(SmartLockInfos[4]);
+		Lock.Branch = MoveTemp(SmartLockInfos[6]);
+		Lock.Status = MoveTemp(SmartLockInfos[8]);
+		Lock.bIsLocked = (Lock.Status == TEXT("Locked"));
+		// Note: keeping the full email address as the owner name so we can display both the short and full name in the tooltip
+		Lock.Owner = MoveTemp(SmartLockInfos[9]);
+		Lock.Workspace = MoveTemp(SmartLockInfos[10]);
+		Lock.Path = MoveTemp(SmartLockInfos[11]);
+	}
+	return Lock;
 }
 
 // Parse the fileinfo output format "{RevisionChangeset};{RevisionHeadChangeset};{RepSpec};{LockedBy};{LockedWhere};{ServerPath}"
@@ -454,6 +497,33 @@ public:
 	FString ServerPath;
 };
 
+/**
+ * Find the locks matching the file path from the list of locks
+ *
+ * Multiple matching locks can only happen if multiple destination branches are configured
+*/
+TArray<FPlasticSourceControlLockRef> FindMatchingLocks(const TArray<FPlasticSourceControlLockRef>& InLocks, const FString& InPath)
+{
+	TArray<FPlasticSourceControlLockRef> MatchingLocks;
+	for (int i = 0; i < InLocks.Num(); i++)
+	{
+		if (InLocks[i]->Path == InPath)
+		{
+			MatchingLocks.Add(InLocks[i]);
+		}
+	}
+	return MatchingLocks;
+}
+
+void ConcatStrings(FString& InOutString, const TCHAR* InSeparator, const FString& InOther)
+{
+	if (!InOutString.IsEmpty())
+	{
+		InOutString += InSeparator;
+	}
+	InOutString += InOther;
+}
+
 /** Parse the array of strings result of a 'cm fileinfo --format="{RevisionChangeset};{RevisionHeadChangeset};{RepSpec};{LockedBy};{LockedWhere}"' command
  *
  * Example cm fileinfo results:
@@ -468,14 +538,15 @@ void ParseFileinfoResults(const TArray<FString>& InResults, TArray<FPlasticSourc
 	ensureMsgf(InResults.Num() == InOutStates.Num(), TEXT("The fileinfo command should gives the same number of infos as the status command"));
 
 	const FPlasticSourceControlProvider& Provider = FPlasticSourceControlModule::Get().GetProvider();
-
+	// Note: here is one of the rare places where we need to use a branch name, not a workspace selector
 	const FString& BranchName = Provider.GetBranchName();
-	const FString& Repository = Provider.GetRepositoryName();
 
-	TMap<FString, FSmartLockInfoParser> SmartLocks;
+	TArray<FPlasticSourceControlLockRef> Locks;
 	if (Provider.GetPlasticScmVersion() >= PlasticSourceControlVersions::SmartLocks)
 	{
-		PlasticSourceControlUtils::RunListSmartLocks(Repository, SmartLocks);
+		// In the Content Browser, only show locks applying to the current working branch
+		const bool bForAllDestBranches = false;
+		PlasticSourceControlUtils::RunListLocks(Provider, bForAllDestBranches, Locks);
 	}
 
 	// Iterate on all files and all status of the result (assuming same number of line of results than number of file states)
@@ -489,22 +560,35 @@ void ParseFileinfoResults(const TArray<FString>& InResults, TArray<FPlasticSourc
 		FileState.LocalRevisionChangeset = FileinfoParser.RevisionChangeset;
 		FileState.DepotRevisionChangeset = FileinfoParser.RevisionHeadChangeset;
 		FileState.RepSpec = FileinfoParser.RepSpec;
-		FileState.LockedBy = MoveTemp(FileinfoParser.LockedBy);
-		FileState.LockedWhere = MoveTemp(FileinfoParser.LockedWhere);
 
-		// Additional information coming from SmartLocks (branch name and "Retained" lock status)
-		FSmartLockInfoParser* SmartLock = SmartLocks.Find(FileinfoParser.ServerPath);
-		if (SmartLock != nullptr)
+		// Additional information coming from Locks (branch, workspace, date and lock status)
+		// Note: in case of multi destination branches, we might have multiple locks for the same path, so we concatenate the string info
+		const TArray<FPlasticSourceControlLockRef> MatchingLocks = FindMatchingLocks(Locks, FileinfoParser.ServerPath);
+		for (auto& Lock : MatchingLocks)
 		{
+			// "Locked" vs "Retained" lock
+			if (Lock->bIsLocked)
+			{
+				ConcatStrings(FileState.LockedBy, TEXT(", "), PlasticSourceControlUtils::UserNameToDisplayName(Lock->Owner));
+			}
 			// Considers a "Retained" lock as meaningful only if it is retained on another branch
-			if ((SmartLock->Status == "Retained") && (SmartLock->BranchName != BranchName))
+			// NOTE: this is required to avoid the Unreal Editor showing a popup warning preventing the user to save the asset
+			else if (Lock->Branch != BranchName)
 			{
-				FileState.RetainedBy = MoveTemp(SmartLock->Owner);
+				ConcatStrings(FileState.RetainedBy, TEXT(", "), PlasticSourceControlUtils::UserNameToDisplayName(Lock->Owner));
 			}
+			ConcatStrings(FileState.LockedWhere, TEXT(", "), Lock->Workspace);
+			ConcatStrings(FileState.LockedBranch, TEXT(", "), Lock->Branch);
 
-			FileState.LockedBranch = MoveTemp(SmartLock->BranchName);
-			FileState.LockedId = SmartLock->ItemId;
-			FileState.LockedDate = SmartLock->Date;
+			// Only save the ItemId if there is only one matching Lock: used to Unlock it from the context menu in the Content Browser,
+			// but leave the ItmeId to invalid if there are more than one: there would be no way to know which one to unlock from the context menu
+			// (Unlocking in such a case require using the View Locks window instead for disambiguation)
+			if (MatchingLocks.Num() == 1)
+			{
+				FileState.LockedId = Lock->ItemId;
+			}
+			// Note; this will keep only the date of the last lock
+			FileState.LockedDate = Lock->Date;
 		}
 
 		// debug log (only for the first few files)
@@ -546,20 +630,30 @@ FPlasticMergeConflictParser::FPlasticMergeConflictParser(const FString& InResult
 	}
 }
 
+// Types of changes in source control revisions, using Perforce terminology for the History window
+static const TCHAR* SourceControlActionAdded = TEXT("add");
+static const TCHAR* SourceControlActionDeleted = TEXT("delete");
+static const TCHAR* SourceControlActionMoved = TEXT("branch");
+static const TCHAR* SourceControlActionMerged = TEXT("integrate");
+static const TCHAR* SourceControlActionChanged = TEXT("edit");
+
 // Convert a file state to a string ala Perforce, see also ParseShelveFileStatus()
 FString FileStateToAction(const EWorkspaceState InState)
 {
 	switch (InState)
 	{
 	case EWorkspaceState::Added:
-		return TEXT("add");
+		return SourceControlActionAdded;
 	case EWorkspaceState::Deleted:
-		return TEXT("delete");
+		return SourceControlActionDeleted;
 	case EWorkspaceState::Moved:
-		return TEXT("branch");
+	case EWorkspaceState::Copied:
+		return SourceControlActionMoved;
+	case EWorkspaceState::Replaced:
+		return SourceControlActionMerged;
 	case EWorkspaceState::CheckedOutChanged:
 	default:
-		return TEXT("edit");
+		return SourceControlActionChanged;
 	}
 }
 
@@ -579,9 +673,34 @@ static FString DecodeXmlEntities(const FString& InString)
 	return String;
 }
 
+/**
+ * Parse one specific node of the result of the 'cm history --moveddeleted' command.
+ *
+ * Results of the history command in case of a move looks like that:
+ <Branch>Moved from /Content/FirstPersonBP/Blueprints/BP_ToRename.uasset to /Content/FirstPersonBP/Blueprints/BP_TestsRenamed.uasset</Branch>
+*/
+static FString ParseMovedFrom(const FXmlNode* InBranchNode)
+{
+	FString MovedFrom;
+
+	if (InBranchNode != nullptr)
+	{
+		static const int32 MovedFromPrefixLen = FString("Moved from /").Len();
+		MovedFrom = InBranchNode->GetContent().RightChop(MovedFromPrefixLen);
+
+		const int32 MovedToIndex = MovedFrom.Find(TEXT(" to "), ESearchCase::CaseSensitive);
+		if (MovedToIndex != INDEX_NONE)
+		{
+			MovedFrom.LeftInline(MovedToIndex);
+		}
+	}
+
+	return MovedFrom; // Convert server path to absolute
+}
+
 /**
  * Parse results of the 'cm history --moveddeleted --xml --encoding="utf-8"' command.
- * 
+ *
  * Results of the history command looks like that:
 <RevisionHistoriesResult>
   <RevisionHistories>
@@ -594,18 +713,37 @@ static FString DecodeXmlEntities(const FString& InString)
 		  <CreationDate>2019-10-14T09:52:07+02:00</CreationDate>
 		  <RevisionType>bin</RevisionType>
 		  <ChangesetNumber>7</ChangesetNumber>
-		  <Owner>SRombauts</Owner>
+		  <Owner>sebastien.rombauts</Owner>
 		  <Comment>New tests</Comment>
 		  <Repository>UE4PlasticPluginDev</Repository>
 		  <Server>localhost:8087</Server>
 		  <RepositorySpec>UE4PlasticPluginDev@localhost:8087</RepositorySpec>
+		  <DataStatus>Available</DataStatus>
+		  <ItemId>1657</ItemId>
 		  <Size>22356</Size>
 		  <Hash>zzuB6G9fbWz1md12+tvBxg==</Hash>
 		</Revision>
 		...
 		<Revision>
 		  <RevisionSpec>C:/Workspace/UE4PlasticPluginDev/Content/FirstPersonBP/Blueprints/BP_TestsRenamed.uasset#cs:12</RevisionSpec>
-		  <Branch>Removed /Content/FirstPersonBP/Blueprints/BP_TestsRenamed.uasset</Branch>
+		  <Branch>/main/rename_test</Branch>
+		  <CreationDate>2022-04-28T16:00:37+02:00</CreationDate>
+		  <RevisionType>bin</RevisionType>
+		  <ChangesetNumber>12</ChangesetNumber>
+		  <Owner>sebastien.rombauts</Owner>
+		  <Comment>Renamed sphere blueprint</Comment>
+		  <Repository>UE4PlasticPluginDev</Repository>
+		  <Server>localhost:8087</Server>
+		  <RepositorySpec>UE4PlasticPluginDev@localhost:8087</RepositorySpec>
+		  <DataStatus>Available</DataStatus>
+		  <ItemPathOrSpec>C:/Workspace/UE4PlasticPluginDev/Content/FirstPersonBP/Blueprints/BP_TestsRenamed.uasset</ItemPathOrSpec>
+		  <ItemId>1657</ItemId>
+		  <Size>28603</Size>
+		  <Hash>MREVIZ1qKNqu1h2iq9WiRg==</Hash>
+		</Revision>
+		<Revision>
+		  <RevisionSpec>C:/Workspace/UE4PlasticPluginDev/Content/FirstPersonBP/Blueprints/BP_TestsRenamed.uasset#cs:12</RevisionSpec>
+		  <Branch>Moved from /Content/FirstPersonBP/Blueprints/BP_ToRename.uasset to /Content/FirstPersonBP/Blueprints/BP_TestsRenamed.uasset</Branch>
 		  <CreationDate>2022-04-28T16:00:37+02:00</CreationDate>
 		  <RevisionType />
 		  <ChangesetNumber>12</ChangesetNumber>
@@ -614,10 +752,13 @@ static FString DecodeXmlEntities(const FString& InString)
 		  <Repository>UE4PlasticPluginDev</Repository>
 		  <Server>localhost:8087</Server>
 		  <RepositorySpec>UE4PlasticPluginDev@localhost:8087</RepositorySpec>
-		  <Size>22406</Size>
-		  <Hash>uR7NdDRAyKqADdyAqh67Rg==</Hash>
+		  <DataStatus />
+		  <ItemPathOrSpec />
+		  <ItemId />
+		  <Size>0</Size>
+		  <Hash />
 		</Revision>
-
+		...
 	  </Revisions>
 	</RevisionHistory>
 	<RevisionHistory>
@@ -629,10 +770,10 @@ static FString DecodeXmlEntities(const FString& InString)
 */
 static bool ParseHistoryResults(const bool bInUpdateHistory, const FXmlFile& InXmlResult, TArray<FPlasticSourceControlState>& InOutStates)
 {
-	TRACE_CPUPROFILER_EVENT_SCOPE(PlasticSourceControlParsers::ParseHistoryResults);
-
 	const FPlasticSourceControlProvider& Provider = FPlasticSourceControlModule::Get().GetProvider();
+	const FString& WorkspaceRoot = Provider.GetPathToWorkspaceRoot();
 	const FString RootRepSpec = FString::Printf(TEXT("%s@%s"), *Provider.GetRepositoryName(), *Provider.GetServerUrl());
+	const FString CurrentBranch = Provider.GetBranchName();
 
 	static const FString RevisionHistoriesResult(TEXT("RevisionHistoriesResult"));
 	static const FString RevisionHistories(TEXT("RevisionHistories"));
@@ -670,7 +811,7 @@ static bool ParseHistoryResults(const bool bInUpdateHistory, const FXmlFile& InX
 			continue;
 		}
 
-		const FString Filename = ItemNameNode->GetContent();
+		FString Filename = ItemNameNode->GetContent();
 		FPlasticSourceControlState* InOutStatePtr = InOutStates.FindByPredicate(
 			[&Filename](const FPlasticSourceControlState& State) { return State.LocalFilename == Filename; }
 		);
@@ -692,112 +833,127 @@ static bool ParseHistoryResults(const bool bInUpdateHistory, const FXmlFile& InX
 			InOutState.History.Reserve(RevisionNodes.Num());
 		}
 
-		// parse history in reverse: needed to get most recent at the top (implied by the UI)
-		// Note: limit to last 100 changes, like Perforce
-		static const int32 MaxRevisions = 100;
-		const int32 MinIndex = FMath::Max(0, RevisionNodes.Num() - MaxRevisions);
-		for (int32 Index = RevisionNodes.Num() - 1; Index >= MinIndex; Index--)
+		// parse history in reverse: needed to get most recent at the top (required by Unreal Editor for the "Diff with depot" using the index 0)
+		FString NextEntryMovedFrom;
+		for (int32 RevisionIndex = RevisionNodes.Num() - 1; RevisionIndex >= 0; RevisionIndex--)
 		{
-			if (const FXmlNode* RevisionNode = RevisionNodes[Index])
+			const FXmlNode* RevisionNode = RevisionNodes[RevisionIndex];
+			check(RevisionNode);
+
+			const TSharedRef<FPlasticSourceControlRevision, ESPMode::ThreadSafe> SourceControlRevision = MakeShareable(new FPlasticSourceControlRevision);
+			SourceControlRevision->State = &InOutState;
+			SourceControlRevision->Filename = Filename;
+
+			if (const FXmlNode* RevisionTypeNode = RevisionNode->FindChildNode(RevisionType))
 			{
-				const TSharedRef<FPlasticSourceControlRevision, ESPMode::ThreadSafe> SourceControlRevision = MakeShared<FPlasticSourceControlRevision>();
-				SourceControlRevision->State = &InOutState;
-				SourceControlRevision->Filename = Filename;
+				// There are two entries for a Move of an asset;
+				// 1. a regular one with the normal data: revision, comment, branch, Id, size, hash etc.
+				// 2. and another "empty" one for the Move
+				// => Since the parsing is done in reverse order, the detection of a Move need to apply to the next entry
+				if (RevisionTypeNode->GetContent().IsEmpty())
+				{
+					// An empty <RevisionType> signals a Move: save the "MovedFrom" filename to treat the next entry as a Move and update the Filename accordingly for next (older) entries
+					NextEntryMovedFrom = FPaths::Combine(WorkspaceRoot, ParseMovedFrom(RevisionNode->FindChildNode(Branch)));
 
-				if (const FXmlNode* RevisionTypeNode = RevisionNode->FindChildNode(RevisionType))
+					// and skip this revision as it is empty (it's just an additional entry with data for the move)
+					continue;
+				}
+				else
 				{
-					if (!RevisionTypeNode->GetContent().IsEmpty())
-					{
-						if (Index == 0)
-						{
-							SourceControlRevision->Action = FileStateToAction(EWorkspaceState::Added);
-						}
-						else
-						{
-							SourceControlRevision->Action = FileStateToAction(EWorkspaceState::CheckedOutChanged);
-						}
-					}
-					else
+					// If this entry was flagged as a move:
+					if (NextEntryMovedFrom.Len() > 0)
 					{
-						SourceControlRevision->Action = FileStateToAction(EWorkspaceState::Deleted);
+						// Set this revision as a Move
+						SourceControlRevision->Action = SourceControlActionMoved;
+						// Update Filename for next (older) entries in the history and clear the NextEntryMovedFrom used as a flag
+						Filename = MoveTemp(NextEntryMovedFrom);
 					}
-				}
-
-				if (const FXmlNode* ChangesetNumberNode = RevisionNode->FindChildNode(ChangesetNumber))
-				{
-					const FString& Changeset = ChangesetNumberNode->GetContent();
-					SourceControlRevision->ChangesetNumber = FCString::Atoi(*Changeset); // Value now used in the Revision column and in the Asset Menu History
-
-					// Also append depot name to the revision, but only when it is different from the default one (ie for xlinks sub repository)
-					if (!InOutState.RepSpec.IsEmpty() && (InOutState.RepSpec != RootRepSpec))
+					else if (RevisionIndex == 0)
 					{
-						TArray<FString> RepSpecs;
-						InOutState.RepSpec.ParseIntoArray(RepSpecs, TEXT("@"));
-						SourceControlRevision->Revision = FString::Printf(TEXT("cs:%s@%s"), *Changeset, *RepSpecs[0]);
+						SourceControlRevision->Action = SourceControlActionAdded;
 					}
 					else
 					{
-						SourceControlRevision->Revision = FString::Printf(TEXT("cs:%s"), *Changeset);
+						SourceControlRevision->Action = SourceControlActionChanged;
 					}
 				}
-				if (const FXmlNode* CommentNode = RevisionNode->FindChildNode(Comment))
-				{
-					SourceControlRevision->Description = DecodeXmlEntities(CommentNode->GetContent());
-				}
-				if (const FXmlNode* OwnerNode = RevisionNode->FindChildNode(Owner))
-				{
-					SourceControlRevision->UserName = PlasticSourceControlUtils::UserNameToDisplayName(OwnerNode->GetContent());
-				}
-				if (const FXmlNode* DateNode = RevisionNode->FindChildNode(CreationDate))
-				{
-					const FString& DateIso = DateNode->GetContent();
-					FDateTime::ParseIso8601(*DateIso, SourceControlRevision->Date);
-				}
-				if (const FXmlNode* BranchNode = RevisionNode->FindChildNode(Branch))
-				{
-					SourceControlRevision->Branch = DecodeXmlEntities(BranchNode->GetContent());
-				}
-				if (const FXmlNode* SizeNode = RevisionNode->FindChildNode(Size))
-				{
-					SourceControlRevision->FileSize = FCString::Atoi(*SizeNode->GetContent());
-				}
+			}
 
-				// A negative RevisionHeadChangeset provided by fileinfo mean that the file has been unshelved;
-				// replace it by the changeset number of the first revision in the history (the more recent)
-				// Note: workaround to be able to show the history / the diff of a file that has been unshelved
-				// (but keeps the LocalRevisionChangeset to the negative changeset corresponding to the Shelve Id)
-				if (InOutState.DepotRevisionChangeset < 0)
-				{
-					InOutState.DepotRevisionChangeset = SourceControlRevision->ChangesetNumber;
-				}
+			if (const FXmlNode* ChangesetNumberNode = RevisionNode->FindChildNode(ChangesetNumber))
+			{
+				const FString& Changeset = ChangesetNumberNode->GetContent();
+				SourceControlRevision->ChangesetNumber = FCString::Atoi(*Changeset); // Value now used in the Revision column and in the Asset Menu History
 
-				// Detect and skip more recent changesets on other branches (ie above the RevisionHeadChangeset)
-				// since we usually don't want to display changes from other branches in the History window...
-				// except in case of a merge conflict, where the Editor expects the tip of the "source (remote)" branch to be at the top of the history!
-				if (   (SourceControlRevision->ChangesetNumber > InOutState.DepotRevisionChangeset)
-					&& (SourceControlRevision->GetRevision() != InOutState.PendingResolveInfo.RemoteRevision))
+				// Also append depot name to the revision, but only when it is different from the default one (ie for xlinks sub repository)
+				if (!InOutState.RepSpec.IsEmpty() && (InOutState.RepSpec != RootRepSpec))
 				{
-					InOutState.HeadBranch = SourceControlRevision->Branch;
-					InOutState.HeadAction = SourceControlRevision->Action;
-					InOutState.HeadChangeList = SourceControlRevision->ChangesetNumber;
-					InOutState.HeadUserName = SourceControlRevision->UserName;
-					InOutState.HeadModTime = SourceControlRevision->Date.ToUnixTimestamp();
+					TArray<FString> RepSpecs;
+					InOutState.RepSpec.ParseIntoArray(RepSpecs, TEXT("@"));
+					SourceControlRevision->Revision = FString::Printf(TEXT("cs:%s@%s"), *Changeset, *RepSpecs[0]);
 				}
-				else if (bInUpdateHistory)
+				else
 				{
-					InOutState.History.Add(SourceControlRevision);
+					SourceControlRevision->Revision = FString::Printf(TEXT("cs:%s"), *Changeset);
 				}
+			}
+			if (const FXmlNode* CommentNode = RevisionNode->FindChildNode(Comment))
+			{
+				SourceControlRevision->Description = DecodeXmlEntities(CommentNode->GetContent());
+			}
+			if (const FXmlNode* OwnerNode = RevisionNode->FindChildNode(Owner))
+			{
+				SourceControlRevision->UserName = PlasticSourceControlUtils::UserNameToDisplayName(OwnerNode->GetContent());
+			}
+			if (const FXmlNode* DateNode = RevisionNode->FindChildNode(CreationDate))
+			{
+				const FString& DateIso = DateNode->GetContent();
+				FDateTime::ParseIso8601(*DateIso, SourceControlRevision->Date);
+			}
+			if (const FXmlNode* BranchNode = RevisionNode->FindChildNode(Branch))
+			{
+				SourceControlRevision->Branch = DecodeXmlEntities(BranchNode->GetContent());
+			}
+			if (const FXmlNode* SizeNode = RevisionNode->FindChildNode(Size))
+			{
+				SourceControlRevision->FileSize = FCString::Atoi(*SizeNode->GetContent());
+			}
 
-				// Also grab the UserName of the author of the current depot/head changeset
-				if ((SourceControlRevision->ChangesetNumber == InOutState.DepotRevisionChangeset) && InOutState.HeadUserName.IsEmpty())
-				{
-					InOutState.HeadUserName = SourceControlRevision->UserName;
-				}
+			// A negative RevisionHeadChangeset provided by fileinfo mean that the file has been unshelved;
+			// replace it by the changeset number of the first revision in the history (the more recent)
+			// Note: workaround to be able to show the history / the diff of a file that has been unshelved
+			// (but keeps the LocalRevisionChangeset to the negative changeset corresponding to the Shelve Id)
+			if (InOutState.DepotRevisionChangeset < 0)
+			{
+				InOutState.DepotRevisionChangeset = SourceControlRevision->ChangesetNumber;
+			}
 
-				if (!bInUpdateHistory)
-				{
-					break; // if not updating the history, just getting the head of the latest branch is enough
-				}
+			// Detect and skip more recent changesets on other branches (ie above the RevisionHeadChangeset)
+			// since we usually don't want to display changes from other branches in the History window...
+			// except in case of a merge conflict, where the Editor expects the tip of the "source (remote)" branch to be at the top of the history!
+			if (   (SourceControlRevision->ChangesetNumber > InOutState.DepotRevisionChangeset)
+				&& (SourceControlRevision->Branch != CurrentBranch)
+				&& (SourceControlRevision->GetRevision() != InOutState.PendingResolveInfo.RemoteRevision))
+			{
+				InOutState.HeadBranch = SourceControlRevision->Branch;
+				InOutState.HeadAction = SourceControlRevision->Action;
+				InOutState.HeadChangeList = SourceControlRevision->ChangesetNumber;
+				InOutState.HeadUserName = SourceControlRevision->UserName;
+				InOutState.HeadModTime = SourceControlRevision->Date.ToUnixTimestamp();
+			}
+			else if (bInUpdateHistory)
+			{
+				InOutState.History.Add(SourceControlRevision);
+			}
+
+			// Also grab the UserName of the author of the current depot/head changeset
+			if ((SourceControlRevision->ChangesetNumber == InOutState.DepotRevisionChangeset) && InOutState.HeadUserName.IsEmpty())
+			{
+				InOutState.HeadUserName = SourceControlRevision->UserName;
+			}
+
+			if (!bInUpdateHistory)
+			{
+				break; // if not updating the history, just getting the head of the latest branch is enough
 			}
 		}
 	}
@@ -805,17 +961,18 @@ static bool ParseHistoryResults(const bool bInUpdateHistory, const FXmlFile& InX
 	return true;
 }
 
-bool ParseHistoryResults(const bool bInUpdateHistory, const FString& InResults, TArray<FPlasticSourceControlState>& InOutStates)
+bool ParseHistoryResults(const bool bInUpdateHistory, const FString& InXmlFilename, TArray<FPlasticSourceControlState>& InOutStates)
 {
 	bool bResult = false;
 
 	FXmlFile XmlFile;
 	{
-		TRACE_CPUPROFILER_EVENT_SCOPE(PlasticSourceControlParsers::ParseHistoryResults::FXmlFile::LoadFile);
-		bResult = XmlFile.LoadFile(InResults, EConstructMethod::ConstructFromBuffer);
+		TRACE_CPUPROFILER_EVENT_SCOPE(PlasticSourceControlParsers::ParseHistoryResults::LoadXml);
+		bResult = XmlFile.LoadFile(InXmlFilename);
 	}
 	if (bResult)
 	{
+		TRACE_CPUPROFILER_EVENT_SCOPE(PlasticSourceControlParsers::ParseHistoryResults::ParseXml);
 		bResult = ParseHistoryResults(bInUpdateHistory, XmlFile, InOutStates);
 	}
 	else
@@ -827,7 +984,7 @@ bool ParseHistoryResults(const bool bInUpdateHistory, const FString& InResults,
 }
 
 /* Parse results of the 'cm update --xml=tempfile.xml --encoding="utf-8"' command.
- * 
+ *
  * Results of the update command looks like that:
 <UpdatedItems>
   <List>
@@ -842,8 +999,6 @@ bool ParseHistoryResults(const bool bInUpdateHistory, const FString& InResults,
 */
 static bool ParseUpdateResults(const FXmlFile& InXmlResult, TArray<FString>& OutFiles)
 {
-	TRACE_CPUPROFILER_EVENT_SCOPE(PlasticSourceControlParsers::ParseUpdateResults);
-
 	static const FString UpdatedItems(TEXT("UpdatedItems"));
 	static const FString List(TEXT("List"));
 	static const FString UpdatedItem(TEXT("UpdatedItem"));
@@ -884,11 +1039,12 @@ bool ParseUpdateResults(const FString& InResults, TArray<FString>& OutFiles)
 
 	FXmlFile XmlFile;
 	{
-		TRACE_CPUPROFILER_EVENT_SCOPE(PlasticSourceControlParsers::ParseUpdateResults::FXmlFile::LoadFile);
+		TRACE_CPUPROFILER_EVENT_SCOPE(PlasticSourceControlParsers::ParseUpdateResults::LoadXml);
 		bResult = XmlFile.LoadFile(InResults, EConstructMethod::ConstructFromBuffer);
 	}
 	if (bResult)
 	{
+		TRACE_CPUPROFILER_EVENT_SCOPE(PlasticSourceControlParsers::ParseUpdateResults::ParseXml);
 		bResult = ParseUpdateResults(XmlFile, OutFiles);
 	}
 	else
@@ -935,6 +1091,8 @@ bool ParseUpdateResults(const TArray<FString>& InResults, TArray<FString>& OutFi
 /// Parse checkin result, usually looking like "Created changeset cs:8@br:/main@MyProject@SRombauts@cloud (mount:'/')"
 FText ParseCheckInResults(const TArray<FString>& InResults)
 {
+	TRACE_CPUPROFILER_EVENT_SCOPE(PlasticSourceControlParsers::ParseCheckInResults);
+
 	if (InResults.Num() > 0)
 	{
 		static const FString ChangesetPrefix(TEXT("Created changeset "));
@@ -994,8 +1152,6 @@ FText ParseCheckInResults(const TArray<FString>& InResults)
 */
 static bool ParseChangelistsResults(const FXmlFile& InXmlResult, TArray<FPlasticSourceControlChangelistState>& OutChangelistsStates, TArray<TArray<FPlasticSourceControlState>>& OutCLFilesStates)
 {
-	TRACE_CPUPROFILER_EVENT_SCOPE(PlasticSourceControlParsers::ParseChangelistsResults);
-
 	static const FString StatusOutput(TEXT("StatusOutput"));
 	static const FString WkConfigType(TEXT("WkConfigType"));
 	static const FString WkConfigName(TEXT("WkConfigName"));
@@ -1007,8 +1163,11 @@ static bool ParseChangelistsResults(const FXmlFile& InXmlResult, TArray<FPlastic
 	static const FString Change(TEXT("Change"));
 	static const FString Type(TEXT("Type"));
 	static const FString Path(TEXT("Path"));
+	static const FString OldPath(TEXT("OldPath"));
 
-	const FString& WorkspaceRoot = FPlasticSourceControlModule::Get().GetProvider().GetPathToWorkspaceRoot();
+	FPlasticSourceControlProvider& Provider = FPlasticSourceControlModule::Get().GetProvider();
+	const FString& WorkspaceRoot = Provider.GetPathToWorkspaceRoot();
+	const bool bUsesCheckedOutChanged = Provider.GetPlasticScmVersion() >= PlasticSourceControlVersions::StatusIsCheckedOutChanged;
 
 	const FXmlNode* StatusOutputNode = InXmlResult.GetRootNode();
 	if (StatusOutputNode == nullptr || StatusOutputNode->GetTag() != StatusOutput)
@@ -1048,13 +1207,33 @@ static bool ParseChangelistsResults(const FXmlFile& InXmlResult, TArray<FPlastic
 					continue;
 				}
 
-				// Here we make sure to only collect file states, not directories, since we shouldn't display the added directories to the Editor
-				FString FileName = PathNode->GetContent();
-				int32 DotIndex;
-				if (FileName.FindChar(TEXT('.'), DotIndex))
+				FPlasticSourceControlState FileState(FPaths::ConvertRelativePathToFull(WorkspaceRoot, PathNode->GetContent()));
+				FileState.Changelist = ChangelistState.Changelist;
+				if (const FXmlNode* TypeNode = ChangeNode->FindChildNode(Type))
+				{
+					FileState.WorkspaceState = StateFromStatus(TypeNode->GetContent(), bUsesCheckedOutChanged);
+				}
+
+				if (FileState.WorkspaceState == EWorkspaceState::Moved)
+				{
+					if (const FXmlNode* OldPathNode = ChangeNode->FindChildNode(OldPath))
+					{
+						FileState.MovedFrom = FPaths::ConvertRelativePathToFull(WorkspaceRoot, OldPathNode->GetContent());
+					}
+				}
+
+				// Note: in case of a Moved file, it appears twice in the list; just update the first entry (set as a "Changed") with the "Move" status
+				if (FPlasticSourceControlState* ExistingState = OutCLFilesStates[ChangelistIndex].FindByPredicate(
+					[&FileState](const FPlasticSourceControlState& InState)
+					{
+						return InState.GetFilename().Equals(FileState.GetFilename());
+					}))
+				{
+					ExistingState->WorkspaceState = FileState.WorkspaceState;
+					ExistingState->MovedFrom = FileState.MovedFrom;
+				}
+				else
 				{
-					FPlasticSourceControlState FileState(FPaths::ConvertRelativePathToFull(WorkspaceRoot, MoveTemp(FileName)));
-					FileState.Changelist = ChangelistState.Changelist;
 					OutCLFilesStates[ChangelistIndex].Add(MoveTemp(FileState));
 				}
 			}
@@ -1076,17 +1255,18 @@ static bool ParseChangelistsResults(const FXmlFile& InXmlResult, TArray<FPlastic
 	return true;
 }
 
-bool ParseChangelistsResults(const FString& Results, TArray<FPlasticSourceControlChangelistState>& OutChangelistsStates, TArray<TArray<FPlasticSourceControlState>>& OutCLFilesStates)
+bool ParseChangelistsResults(const FString& InXmlFilename, TArray<FPlasticSourceControlChangelistState>& OutChangelistsStates, TArray<TArray<FPlasticSourceControlState>>& OutCLFilesStates)
 {
 	bool bResult = false;
 
 	FXmlFile XmlFile;
 	{
-		TRACE_CPUPROFILER_EVENT_SCOPE(PlasticSourceControlParsers::ParseChangelistsResults::FXmlFile::LoadFile);
-		bResult = XmlFile.LoadFile(Results, EConstructMethod::ConstructFromBuffer);
+		TRACE_CPUPROFILER_EVENT_SCOPE(PlasticSourceControlParsers::ParseChangelistsResults::LoadXml);
+		bResult = XmlFile.LoadFile(InXmlFilename);
 	}
 	if (bResult)
 	{
+		TRACE_CPUPROFILER_EVENT_SCOPE(PlasticSourceControlParsers::ParseChangelistsResults::ParseXml);
 		bResult = ParseChangelistsResults(XmlFile, OutChangelistsStates, OutCLFilesStates);
 	}
 	else
@@ -1251,11 +1431,12 @@ bool ParseShelvesResults(const FString& InResults, TArray<FPlasticSourceControlC
 
 	FXmlFile XmlFile;
 	{
-		TRACE_CPUPROFILER_EVENT_SCOPE(PlasticSourceControlParsers::ParseShelvesResults);
+		TRACE_CPUPROFILER_EVENT_SCOPE(PlasticSourceControlParsers::ParseShelvesResults::LoadXml);
 		bResult = XmlFile.LoadFile(InResults, EConstructMethod::ConstructFromBuffer);
 	}
 	if (bResult)
 	{
+		TRACE_CPUPROFILER_EVENT_SCOPE(PlasticSourceControlParsers::ParseShelvesResults::ParseXml);
 		bResult = ParseShelvesResults(XmlFile, InOutChangelistsStates);
 	}
 	else
@@ -1271,12 +1452,14 @@ bool ParseShelvesResults(const FString& InResults, TArray<FPlasticSourceControlC
  *
  * Results of the diff command looks like that:
 C;666;Content\NewFolder\BP_CheckedOut.uasset
- * but for Moved assets there are two entires that we need to merge:
+ * but for Moved assets there are two entries that we need to merge:
 C;266;"Content\ThirdPerson\Blueprints\BP_ThirdPersonCharacterRenamed.uasset"
 M;-1;"Content\ThirdPerson\Blueprints\BP_ThirdPersonCharacterRenamed.uasset"
 */
 bool ParseShelveDiffResults(const FString InWorkspaceRoot, TArray<FString>&& InResults, TArray<FPlasticSourceControlRevision>& OutBaseRevisions)
 {
+	TRACE_CPUPROFILER_EVENT_SCOPE(PlasticSourceControlParsers::ParseShelveDiffResults);
+
 	bool bResult = true;
 
 	OutBaseRevisions.Reset(InResults.Num());
@@ -1296,14 +1479,14 @@ bool ParseShelveDiffResults(const FString InWorkspaceRoot, TArray<FString>&& InR
 
 			if (ShelveState == EWorkspaceState::Moved)
 			{
-				// In case of a Moved file, it appears twice in the list, so update the first entry (set as a "Changed" but has the Base Revision Id) and update it with the "Move" status
+				// Note: in case of a Moved file, it appears twice in the list; just update the first entry (set as a "Changed") with the "Move" status
 				if (FPlasticSourceControlRevision* ExistingShelveRevision = OutBaseRevisions.FindByPredicate(
 					[&AbsoluteFilename](const FPlasticSourceControlRevision& State)
 					{
 						return State.GetFilename().Equals(AbsoluteFilename);
 					}))
 				{
-					ExistingShelveRevision->Action = FileStateToAction(EWorkspaceState::Moved);
+					ExistingShelveRevision->Action = SourceControlActionMoved;
 					continue;
 				}
 			}
@@ -1390,11 +1573,12 @@ bool ParseShelvesResult(const FString& InResults, FString& OutComment, FDateTime
 
 	FXmlFile XmlFile;
 	{
-		TRACE_CPUPROFILER_EVENT_SCOPE(PlasticSourceControlParsers::ParseShelvesResult);
+		TRACE_CPUPROFILER_EVENT_SCOPE(PlasticSourceControlParsers::ParseShelvesResult::LoadXml);
 		bResult = XmlFile.LoadFile(InResults, EConstructMethod::ConstructFromBuffer);
 	}
 	if (bResult)
 	{
+		TRACE_CPUPROFILER_EVENT_SCOPE(PlasticSourceControlParsers::ParseShelvesResult::ParseXml);
 		int32 ShelveId;
 		bResult = PlasticSourceControlParsers::ParseShelvesResult(XmlFile, ShelveId, OutComment, OutDate, OutOwner);
 	}
@@ -1406,6 +1590,326 @@ bool ParseShelvesResult(const FString& InResults, FString& OutComment, FDateTime
 	return bResult;
 }
 
+
+/**
+ * Parse results of the 'cm find changeset --xml --encoding="utf-8"' command.
+ *
+ * Results of the find command looks like the following:
+<?xml version="1.0" encoding="utf-8" ?>
+<PLASTICQUERY>
+  <CHANGESET>
+    <ID>2652</ID>
+    <CHANGESETID>56</CHANGESETID>
+    <COMMENT>test</COMMENT>
+    <DATE>2024-03-25T10:37:14+01:00</DATE>
+    <OWNER>[REDACTED_EMAIL]</OWNER>
+    <REPOSITORY>UE5PlasticPluginDev</REPOSITORY>
+    <REPNAME>UE5PlasticPluginDev</REPNAME>
+    <REPSERVER>SRombautsU@cloud</REPSERVER>
+    <BRANCH>/main</BRANCH>
+    <PARENT>55</PARENT>
+    <GUID>d49c552e-9654-44d0-86eb-0d55fa5e8dc3</GUID>
+    <ROOTREV>2651</ROOTREV>
+  </CHANGESET>
+  [...]
+</PLASTICQUERY>
+*/
+static bool ParseChangesetesResults(const FXmlFile& InXmlResult, TArray<FPlasticSourceControlChangesetRef>& OutChangesets)
+{
+	static const FString PlasticQuery(TEXT("PLASTICQUERY"));
+	static const FString ChangesetId(TEXT("CHANGESETID"));
+	static const FString Branch(TEXT("BRANCH"));
+	static const FString Comment(TEXT("COMMENT"));
+	static const FString Owner(TEXT("OWNER"));
+	static const FString Date(TEXT("DATE"));
+
+	const FXmlNode* PlasticQueryNode = InXmlResult.GetRootNode();
+	if (PlasticQueryNode == nullptr || PlasticQueryNode->GetTag() != PlasticQuery)
+	{
+		return false;
+	}
+
+	const TArray<FXmlNode*>& ChangesetsNodes = PlasticQueryNode->GetChildrenNodes();
+	OutChangesets.Reserve(ChangesetsNodes.Num());
+	for (const FXmlNode* ChangesetNode : ChangesetsNodes)
+	{
+		check(ChangesetNode);
+		const FXmlNode* ChangesetIdNode = ChangesetNode->FindChildNode(ChangesetId);
+		if (ChangesetIdNode == nullptr)
+		{
+			continue;
+		}
+
+		FPlasticSourceControlChangesetRef ChangesetRef = MakeShareable(new FPlasticSourceControlChangeset());
+
+		ChangesetRef->ChangesetId = FCString::Atoi(*ChangesetIdNode->GetContent());
+
+		if (const FXmlNode* CommentNode = ChangesetNode->FindChildNode(Comment))
+		{
+			ChangesetRef->Comment = DecodeXmlEntities(CommentNode->GetContent());
+		}
+		if (const FXmlNode* BranchNode = ChangesetNode->FindChildNode(Branch))
+		{
+			ChangesetRef->Branch = DecodeXmlEntities(BranchNode->GetContent());
+		}
+		if (const FXmlNode* OwnerNode = ChangesetNode->FindChildNode(Owner))
+		{
+			// Note: keeping the full email address as the owner name so we can display both the short and full name in the tooltip
+			ChangesetRef->CreatedBy = OwnerNode->GetContent();
+		}
+		if (const FXmlNode* DateNode = ChangesetNode->FindChildNode(Date))
+		{
+			const FString& DateIso = DateNode->GetContent();
+			FDateTime::ParseIso8601(*DateIso, ChangesetRef->Date);
+		}
+
+		OutChangesets.Add(MoveTemp(ChangesetRef));
+	}
+
+	return true;
+}
+
+bool ParseChangesetsResults(const FString& InXmlFilename, TArray<FPlasticSourceControlChangesetRef>& OutChangesets)
+{
+	bool bResult = false;
+
+	FXmlFile XmlFile;
+	{
+		TRACE_CPUPROFILER_EVENT_SCOPE(PlasticSourceControlParsers::ParseChangesetesResults::LoadXml);
+		bResult = XmlFile.LoadFile(InXmlFilename);
+	}
+	if (bResult)
+	{
+		TRACE_CPUPROFILER_EVENT_SCOPE(PlasticSourceControlParsers::ParseChangesetesResults::ParseXml);
+		bResult = ParseChangesetesResults(XmlFile, OutChangesets);
+	}
+	else
+	{
+		UE_LOG(LogSourceControl, Error, TEXT("ParseChangesetesResults: XML parse error '%s'"), *XmlFile.GetLastError())
+	}
+
+	return bResult;
+}
+
+
+/**
+ * Convert the type of change in the log of changesets to a state.
+ */
+static EWorkspaceState StateFromType(const FString& InChangeType)
+{
+	static const FString Added(TEXT("Added"));
+	static const FString Changed(TEXT("Changed"));
+	static const FString Deleted(TEXT("Deleted"));
+	static const FString Moved(TEXT("Moved"));
+
+	if (InChangeType.Equals(Changed))
+	{
+		return EWorkspaceState::CheckedOutChanged;
+	}
+	else if (InChangeType.Equals(Added))
+	{
+		return EWorkspaceState::Added;
+	}
+	else if (InChangeType.Equals(Moved))
+	{
+		return EWorkspaceState::Moved;
+	}
+	else if (InChangeType.Equals(Deleted))
+	{
+		return EWorkspaceState::Deleted;
+	}
+
+	return EWorkspaceState::Unknown;
+}
+
+/**
+ * Parse Changes child node in a Changeset.
+ *
+ * Results of the log command looks like the following:
+<?xml version="1.0" encoding="utf-8"?>
+<LogList>
+  <Changeset>
+	[...]
+	<Changes>
+	  <Item>
+		<Branch>/main/test</Branch>
+		<RevNo>72</RevNo>
+		<Owner>[REDACTED_EMAIL]</Owner>
+		<RevId>2861</RevId>
+		<ParentRevId>-1</ParentRevId>
+		<SrcCmPath>/Private/Private.md</SrcCmPath>
+		<SrcParentItemId>2868</SrcParentItemId>
+		<DstCmPath>/Private/Private.md</DstCmPath>
+		<DstParentItemId>2868</DstParentItemId>
+		<Date>2024-04-03T14:59:31+02:00</Date>
+		<Type>Added</Type>
+	  </Item>
+	[...]
+	</Changes>
+	[...]
+  </Changeset>
+</LogList>
+ */
+static void ParseChangesInChangeset(const FXmlNode* InChangesetNode, const FPlasticSourceControlChangesetRef& InChangeset, TArray<FPlasticSourceControlStateRef>& OutFiles)
+{
+	static const FString Changes(TEXT("Changes"));
+	static const FString Item(TEXT("Item"));
+	static const FString Type(TEXT("Type"));
+	static const FString SrcCmPath(TEXT("SrcCmPath"));
+	static const FString DstCmPath(TEXT("DstCmPath"));
+
+	const FPlasticSourceControlProvider& Provider = FPlasticSourceControlModule::Get().GetProvider();
+	const FString RootRepSpec = FString::Printf(TEXT("%s@%s"), *Provider.GetRepositoryName(), *Provider.GetServerUrl());
+
+	if (const FXmlNode* ChangesNode = InChangesetNode->FindChildNode(Changes))
+	{
+		const TArray<FXmlNode*>& ItemNodes = ChangesNode->GetChildrenNodes();
+		OutFiles.Reserve(ItemNodes.Num());
+		for (const FXmlNode* ItemNode : ItemNodes)
+		{
+			check(ItemNode);
+
+			const FXmlNode* PathNode = ItemNode->FindChildNode(DstCmPath);
+			const FXmlNode* TypeNode = ItemNode->FindChildNode(Type);
+			if ((PathNode == nullptr) || (TypeNode == nullptr))
+			{
+				continue;
+			}
+
+			// Note: remove the leading '/' from the server path to make it relative to the root of the workspace
+			FString FileName = PathNode->GetContent().RightChop(1);
+			const EWorkspaceState WorkspaceState = StateFromType(TypeNode->GetContent());
+			FPlasticSourceControlStateRef State = MakeShareable(new FPlasticSourceControlState(MoveTemp(FileName), WorkspaceState));
+			State->RepSpec = RootRepSpec;
+
+			if (WorkspaceState == EWorkspaceState::Moved)
+			{
+				if (const FXmlNode* SrcNode = ItemNode->FindChildNode(SrcCmPath))
+				{
+					State->MovedFrom = SrcNode->GetContent().RightChop(1); // remove the leading '/' character from the server path
+				}
+			}
+
+			// Add one revision to be able to fetch the file content for diff, if it's not marked for deletion.
+			if ((WorkspaceState != EWorkspaceState::Deleted) && (State->History.Num() == 0))
+			{
+				const TSharedRef<FPlasticSourceControlRevision, ESPMode::ThreadSafe> SourceControlRevision = MakeShareable(new FPlasticSourceControlRevision);
+				SourceControlRevision->State = &State.Get();
+				SourceControlRevision->Filename = State->GetFilename();
+				SourceControlRevision->Revision = FString::Printf(TEXT("cs:%d"), InChangeset->ChangesetId);
+				SourceControlRevision->ChangesetNumber = InChangeset->ChangesetId; // Note: for display in the diff window only
+				SourceControlRevision->Date = InChangeset->Date; // Note: not yet used for display as of UE5.2
+
+				State->History.Add(SourceControlRevision);
+			}
+
+			// Note: in case of a Moved file, it appears twice in the list; just update the first entry (set as a "Changed") with the "Move" status
+			if (FPlasticSourceControlStateRef* ExistingState = OutFiles.FindByPredicate(
+				[&State](const TSharedRef<FPlasticSourceControlState, ESPMode::ThreadSafe>& InState)
+				{
+					return InState->GetFilename().Equals(State->GetFilename());
+				}))
+			{
+				(*ExistingState)->WorkspaceState = State->WorkspaceState;
+				(*ExistingState)->MovedFrom = State->MovedFrom;
+			}
+			else
+			{
+				OutFiles.Add(MoveTemp(State));
+			}
+		}
+	}
+}
+
+/**
+ * Parse results of the 'cm log cs:<ChangesetId> --xml --encoding="utf-8"' command.
+ *
+ * Results of the log command looks like the following:
+<?xml version="1.0" encoding="utf-8"?>
+<LogList>
+  <Changeset>
+	<ObjId>2674</ObjId>
+	<ChangesetId>73</ChangesetId>
+	<Branch>/main/test</Branch>
+	<Comment>private files and folders</Comment>
+	<Owner>[REDACTED_EMAIL]</Owner>
+	<GUID>cd803bd1-7d59-4573-b9de-1a4e684d573a</GUID>
+	<Changes>
+	  <Item>
+		<Branch>/main/test</Branch>
+		<RevNo>72</RevNo>
+		<Owner>[REDACTED_EMAIL]</Owner>
+		<RevId>2861</RevId>
+		<ParentRevId>-1</ParentRevId>
+		<SrcCmPath>/Private/Private.md</SrcCmPath>
+		<SrcParentItemId>2868</SrcParentItemId>
+		<DstCmPath>/Private/Private.md</DstCmPath>
+		<DstParentItemId>2868</DstParentItemId>
+		<Date>2024-04-03T14:59:31+02:00</Date>
+		<Type>Added</Type>
+	  </Item>
+	[...]
+	</Changes>
+	<Date>2024-04-02T16:20:11+02:00</Date>
+  </Changeset>
+</LogList>
+*/
+static bool ParseLogResults(const FXmlFile& InXmlResult, const FPlasticSourceControlChangesetRef& InChangeset, TArray<FPlasticSourceControlStateRef>& OutFiles)
+{
+	static const FString LogList(TEXT("LogList"));
+	static const FString ChangesetId(TEXT("ChangesetId"));
+	static const FString Branch(TEXT("Branch"));
+	static const FString Comment(TEXT("Comment"));
+	static const FString Owner(TEXT("Owner"));
+	static const FString Date(TEXT("Date"));
+
+	const FXmlNode* LogListNode = InXmlResult.GetRootNode();
+	if (LogListNode == nullptr || LogListNode->GetTag() != LogList)
+	{
+		return false;
+	}
+
+	const TArray<FXmlNode*>& ChangesetsNodes = LogListNode->GetChildrenNodes();
+	if (ChangesetsNodes.Num() != 1 || ChangesetsNodes[0] == nullptr)
+	{
+		return false;
+	}
+
+	FXmlNode* ChangesetNode = ChangesetsNodes[0];
+	const FXmlNode* ChangesetIdNode = ChangesetNode->FindChildNode(ChangesetId);
+	if (ChangesetIdNode == nullptr || InChangeset->ChangesetId != FCString::Atoi(*ChangesetIdNode->GetContent()))
+	{
+		return false;
+	}
+
+	// List Files States and create a Revision
+	ParseChangesInChangeset(ChangesetNode, InChangeset, OutFiles);
+
+	return true;
+}
+
+bool ParseLogResults(const FString& InXmlFilename, const FPlasticSourceControlChangesetRef& InChangeset, TArray<FPlasticSourceControlStateRef>& OutFiles)
+{
+	bool bResult = false;
+
+	FXmlFile XmlFile;
+	{
+		TRACE_CPUPROFILER_EVENT_SCOPE(PlasticSourceControlParsers::ParseLogResults::LoadXml);
+		bResult = XmlFile.LoadFile(InXmlFilename);
+	}
+	if (bResult)
+	{
+		TRACE_CPUPROFILER_EVENT_SCOPE(PlasticSourceControlParsers::ParseLogResults::ParseXml);
+		bResult = ParseLogResults(XmlFile, InChangeset, OutFiles);
+	}
+	else
+	{
+		UE_LOG(LogSourceControl, Error, TEXT("ParseLogResults: XML parse error '%s'"), *XmlFile.GetLastError())
+	}
+
+	return bResult;
+}
+
 /**
  * Parse results of the 'cm find "branches where date >= 'YYYY-MM-DD' or changesets >= 'YYYY-MM-DD'" --xml --encoding="utf-8"' command.
  *
@@ -1446,8 +1950,9 @@ static bool ParseBranchesResults(const FXmlFile& InXmlResult, TArray<FPlasticSou
 		return false;
 	}
 
-	const TArray<FXmlNode*>& BranchsNodes = PlasticQueryNode->GetChildrenNodes();
-	for (const FXmlNode* BranchNode : BranchsNodes)
+	const TArray<FXmlNode*>& BranchesNodes = PlasticQueryNode->GetChildrenNodes();
+	OutBranches.Reserve(BranchesNodes.Num());
+	for (const FXmlNode* BranchNode : BranchesNodes)
 	{
 		check(BranchNode);
 		const FXmlNode* NameNode = BranchNode->FindChildNode(Name);
@@ -1464,18 +1969,16 @@ static bool ParseBranchesResults(const FXmlFile& InXmlResult, TArray<FPlasticSou
 		{
 			BranchRef->Comment = DecodeXmlEntities(CommentNode->GetContent());
 		}
-
 		if (const FXmlNode* DateNode = BranchNode->FindChildNode(Date))
 		{
 			const FString& DateIso = DateNode->GetContent();
 			FDateTime::ParseIso8601(*DateIso, BranchRef->Date);
 		}
-
 		if (const FXmlNode* OwnerNode = BranchNode->FindChildNode(Owner))
 		{
+			// Note: keeping the full email address as the owner name so we can display both the short and full name in the tooltip
 			BranchRef->CreatedBy = OwnerNode->GetContent();
 		}
-
 		if (const FXmlNode* RepNameNode = BranchNode->FindChildNode(RepName))
 		{
 			if (const FXmlNode* RepServerNode = BranchNode->FindChildNode(RepServer))
@@ -1484,23 +1987,24 @@ static bool ParseBranchesResults(const FXmlFile& InXmlResult, TArray<FPlasticSou
 			}
 		}
 
-		OutBranches.Add(BranchRef);
+		OutBranches.Add(MoveTemp(BranchRef));
 	}
 
 	return true;
 }
 
-bool ParseBranchesResults(const FString& InResults, TArray<FPlasticSourceControlBranchRef>& OutBranches)
+bool ParseBranchesResults(const FString& InXmlFilename, TArray<FPlasticSourceControlBranchRef>& OutBranches)
 {
 	bool bResult = false;
 
 	FXmlFile XmlFile;
 	{
-		TRACE_CPUPROFILER_EVENT_SCOPE(PlasticSourceControlParsers::ParseBranchesResults);
-		bResult = XmlFile.LoadFile(InResults, EConstructMethod::ConstructFromBuffer);
+		TRACE_CPUPROFILER_EVENT_SCOPE(PlasticSourceControlParsers::ParseBranchesResults::LoadXml);
+		bResult = XmlFile.LoadFile(InXmlFilename);
 	}
 	if (bResult)
 	{
+		TRACE_CPUPROFILER_EVENT_SCOPE(PlasticSourceControlParsers::ParseBranchesResults::ParseXml);
 		bResult = ParseBranchesResults(XmlFile, OutBranches);
 	}
 	else
@@ -1535,8 +2039,6 @@ bool ParseBranchesResults(const FString& InResults, TArray<FPlasticSourceControl
 */
 static bool ParseMergeResults(const FXmlFile& InXmlResult, TArray<FString>& OutFiles)
 {
-	TRACE_CPUPROFILER_EVENT_SCOPE(PlasticSourceControlUtils::ParseMergeResults);
-
 	static const FString Merge(TEXT("Merge"));
 	static const FString Added(TEXT("Added"));
 	static const FString Deleted(TEXT("Deleted"));
@@ -1585,11 +2087,12 @@ bool ParseMergeResults(const FString& InResult, TArray<FString>& OutFiles)
 
 	FXmlFile XmlFile;
 	{
-		TRACE_CPUPROFILER_EVENT_SCOPE(PlasticSourceControlParsers::ParseMergeResults::FXmlFile::LoadFile);
+		TRACE_CPUPROFILER_EVENT_SCOPE(PlasticSourceControlParsers::ParseMergeResults::LoadXml);
 		bResult = XmlFile.LoadFile(InResult, EConstructMethod::ConstructFromBuffer);
 	}
 	if (bResult)
 	{
+		TRACE_CPUPROFILER_EVENT_SCOPE(PlasticSourceControlParsers::ParseMergeResults::ParseXml);
 		bResult = ParseMergeResults(XmlFile, OutFiles);
 	}
 	else
diff --git a/Engine/Plugins/Developer/PlasticSourceControl/Source/PlasticSourceControl/Private/PlasticSourceControlParsers.h b/Engine/Plugins/Developer/PlasticSourceControl/Source/PlasticSourceControl/Private/PlasticSourceControlParsers.h
--- a/Engine/Plugins/Developer/PlasticSourceControl/Source/PlasticSourceControl/Private/PlasticSourceControlParsers.h
+++ b/Engine/Plugins/Developer/PlasticSourceControl/Source/PlasticSourceControl/Private/PlasticSourceControlParsers.h
@@ -5,26 +5,17 @@
 #include "CoreMinimal.h"
 
 class FPlasticSourceControlChangelistState;
+class FPlasticSourceControlLock;
 class FPlasticSourceControlRevision;
 class FPlasticSourceControlState;
 typedef TSharedRef<class FPlasticSourceControlBranch, ESPMode::ThreadSafe> FPlasticSourceControlBranchRef;
+typedef TSharedRef<class FPlasticSourceControlChangeset, ESPMode::ThreadSafe> FPlasticSourceControlChangesetRef;
+typedef TSharedRef<class FPlasticSourceControlState, ESPMode::ThreadSafe> FPlasticSourceControlStateRef;
 
 namespace PlasticSourceControlParsers
 {
 
-class FSmartLockInfoParser
-{
-public:
-	explicit FSmartLockInfoParser(const FString& InResult);
-
-	FString Repository;
-	int32 ItemId;
-	FDateTime Date;
-	FString BranchName;
-	FString Status;
-	FString Owner;
-	FString Filename;
-};
+FPlasticSourceControlLock ParseLockInfo(const FString& InResult);
 
 class FPlasticMergeConflictParser
 {
@@ -60,7 +51,9 @@ struct FRemoveRedundantErrors
 	FString Filter;
 };
 
-bool ParseWorkspaceInfo(TArray<FString>& InResults, FString& OutBranchName, FString& OutRepositoryName, FString& OutServerUrl);
+bool ParseProfileInfo(TArray<FString>& InResults, const FString& InServerUrl, FString& OutUserName);
+
+bool ParseWorkspaceInfo(TArray<FString>& InResults, FString& OutWorkspaceSelector, FString& OutBranchName, FString& OutRepositoryName, FString& OutServerUrl);
 
 bool GetChangesetFromWorkspaceStatus(const TArray<FString>& InResults, int32& OutChangeset);
 
@@ -70,14 +63,14 @@ void ParseDirectoryStatusResult(const FString& InDir, const TArray<FString>& InR
 
 void ParseFileinfoResults(const TArray<FString>& InResults, TArray<FPlasticSourceControlState>& InOutStates);
 
-bool ParseHistoryResults(const bool bInUpdateHistory, const FString& InResults, TArray<FPlasticSourceControlState>& InOutStates);
+bool ParseHistoryResults(const bool bInUpdateHistory, const FString& InXmlFilename, TArray<FPlasticSourceControlState>& InOutStates);
 
 bool ParseUpdateResults(const FString& InResults, TArray<FString>& OutFiles);
 bool ParseUpdateResults(const TArray<FString>& InResults, TArray<FString>& OutFiles);
 
 FText ParseCheckInResults(const TArray<FString>& InResults);
 
-bool ParseChangelistsResults(const FString& Results, TArray<FPlasticSourceControlChangelistState>& OutChangelistsStates, TArray<TArray<FPlasticSourceControlState>>& OutCLFilesStates);
+bool ParseChangelistsResults(const FString& InXmlFilename, TArray<FPlasticSourceControlChangelistState>& OutChangelistsStates, TArray<TArray<FPlasticSourceControlState>>& OutCLFilesStates);
 
 bool ParseShelveDiffResult(const FString InWorkspaceRoot, TArray<FString>&& InResults, FPlasticSourceControlChangelistState& InOutChangelistsState);
 bool ParseShelveDiffResults(const FString InWorkspaceRoot, TArray<FString>&& InResults, TArray<FPlasticSourceControlRevision>& OutBaseRevisions);
@@ -85,7 +78,10 @@ bool ParseShelveDiffResults(const FString InWorkspaceRoot, TArray<FString>&& InR
 bool ParseShelvesResults(const FString& InResults, TArray<FPlasticSourceControlChangelistState>& InOutChangelistsStates);
 bool ParseShelvesResult(const FString& InResults, FString& OutComment, FDateTime& OutDate, FString& OutOwner);
 
-bool ParseBranchesResults(const FString& InResults, TArray<FPlasticSourceControlBranchRef>& OutBranches);
+bool ParseChangesetsResults(const FString& InXmlFilename, TArray<FPlasticSourceControlChangesetRef>& OutChangesets);
+bool ParseLogResults(const FString& InXmlFilename, const FPlasticSourceControlChangesetRef& InChangeset, TArray<FPlasticSourceControlStateRef>& OutFiles);
+
+bool ParseBranchesResults(const FString& InXmlFilename, TArray<FPlasticSourceControlBranchRef>& OutBranches);
 
 bool ParseMergeResults(const FString& InResult, TArray<FString>& OutFiles);
 
diff --git a/Engine/Plugins/Developer/PlasticSourceControl/Source/PlasticSourceControl/Private/PlasticSourceControlProjectSettings.h b/Engine/Plugins/Developer/PlasticSourceControl/Source/PlasticSourceControl/Private/PlasticSourceControlProjectSettings.h
--- a/Engine/Plugins/Developer/PlasticSourceControl/Source/PlasticSourceControl/Private/PlasticSourceControlProjectSettings.h
+++ b/Engine/Plugins/Developer/PlasticSourceControl/Source/PlasticSourceControl/Private/PlasticSourceControlProjectSettings.h
@@ -29,19 +29,51 @@ public:
 	UPROPERTY(config, EditAnywhere, Category = "Unity Version Control", meta = (ClampMin = 0))
 	int32 LimitNumberOfRevisionsInHistory = 50;
 
-	/** Show the repository when the branch is created (hidden by default) */
-	UPROPERTY(config, EditAnywhere, Category = "Unity Version Control")
+	/** Show the repository where the branch is created (hidden by default) */
+	UPROPERTY(config, EditAnywhere, Category = "Unity Version Control|View Branches window")
 	bool bShowBranchRepositoryColumn = false;
 
 	/* Show the name of the creator of the branch */
-	UPROPERTY(config, EditAnywhere, Category = "Unity Version Control")
+	UPROPERTY(config, EditAnywhere, Category = "Unity Version Control|View Branches window")
 	bool bShowBranchCreatedByColumn = true;
 
 	/* Show the date of creation of the branch */
-	UPROPERTY(config, EditAnywhere, Category = "Unity Version Control")
+	UPROPERTY(config, EditAnywhere, Category = "Unity Version Control|View Branches window")
 	bool bShowBranchDateColumn = true;
 
 	/* Show the comment of the branch */
-	UPROPERTY(config, EditAnywhere, Category = "Unity Version Control")
+	UPROPERTY(config, EditAnywhere, Category = "Unity Version Control|View Branches window")
 	bool bShowBranchCommentColumn = true;
+
+	/* Show the Id of the lock (hidden by default) */
+	UPROPERTY(config, EditAnywhere, Category = "Unity Version Control|View Locks window")
+	bool bShowLockIdColumn = false;
+
+	/* Show the Workspace of the lock */
+	UPROPERTY(config, EditAnywhere, Category = "Unity Version Control|View Locks window")
+	bool bShowLockWorkspaceColumn = true;
+
+	/* Show the Date of the lock */
+	UPROPERTY(config, EditAnywhere, Category = "Unity Version Control|View Locks window")
+	bool bShowLockDateColumn = true;
+
+	/* Show the Destination Branch of the lock (hidden by default) */
+	UPROPERTY(config, EditAnywhere, Category = "Unity Version Control|View Locks window")
+	bool bShowLockDestinationBranchColumn = false;
+
+	/* Show the name of the creator of the changeset */
+	UPROPERTY(config, EditAnywhere, Category = "Unity Version Control|View Changesets window")
+	bool bShowChangesetCreatedByColumn = true;
+
+	/* Show the date of creation of the changeset */
+	UPROPERTY(config, EditAnywhere, Category = "Unity Version Control|View Changesets window")
+	bool bShowChangesetDateColumn = true;
+
+	/* Show the comment of the changeset */
+	UPROPERTY(config, EditAnywhere, Category = "Unity Version Control|View Changesets window")
+	bool bShowChangesetCommentColumn = true;
+
+	/* Show the branch the changeset was created on */
+	UPROPERTY(config, EditAnywhere, Category = "Unity Version Control|View Changesets window")
+	bool bShowChangesetBranchColumn = true;
 };
diff --git a/Engine/Plugins/Developer/PlasticSourceControl/Source/PlasticSourceControl/Private/PlasticSourceControlProvider.cpp b/Engine/Plugins/Developer/PlasticSourceControl/Source/PlasticSourceControl/Private/PlasticSourceControlProvider.cpp
--- a/Engine/Plugins/Developer/PlasticSourceControl/Source/PlasticSourceControl/Private/PlasticSourceControlProvider.cpp
+++ b/Engine/Plugins/Developer/PlasticSourceControl/Source/PlasticSourceControl/Private/PlasticSourceControlProvider.cpp
@@ -74,7 +74,7 @@ void FPlasticSourceControlProvider::Init(bool bForceConnection)
 		TArray<FString> InfoMessages, ErrorMessages;
 		TArray<FString> Parameters;
 		// Execute a 'checkconnection' command to set bServerAvailable based on the connectivity of the server
-		bServerAvailable = PlasticSourceControlUtils::RunCheckConnection(BranchName, RepositoryName, ServerUrl, InfoMessages, ErrorMessages);
+		bServerAvailable = PlasticSourceControlUtils::RunCheckConnection(WorkspaceSelector, BranchName, RepositoryName, ServerUrl, InfoMessages, ErrorMessages);
 		if (!bServerAvailable)
 		{
 			FMessageLog SourceControlLog("SourceControl");
@@ -126,20 +126,25 @@ void FPlasticSourceControlProvider::CheckPlasticAvailability()
 
 		bUsesLocalReadOnlyState = PlasticSourceControlUtils::GetConfigSetFilesAsReadOnly();
 
-		// Get user name (from the global Unity Version Control client config)
-		PlasticSourceControlUtils::GetUserName(UserName);
-
 		// Register Console Commands
 		PlasticSourceControlConsole.Register();
 
-		if (!bWorkspaceFound)
+		if (bWorkspaceFound)
+		{
+			TArray<FString> ErrorMessages;
+			PlasticSourceControlUtils::GetWorkspaceInfo(WorkspaceSelector, BranchName, RepositoryName, ServerUrl, ErrorMessages);
+			UserName = PlasticSourceControlUtils::GetProfileUserName(ServerUrl);
+		}
+		else
 		{
 			// This info message is only useful here, if bPlasticAvailable, for the Login window
 			FFormatNamedArguments Args;
 			Args.Add(TEXT("WorkspacePath"), FText::FromString(PathToWorkspaceRoot));
 			FMessageLog("SourceControl").Info(FText::Format(LOCTEXT("NotInAWorkspace", "{WorkspacePath} is not in a workspace."), Args));
 
+			// Get default server and user name (from the global client config)
 			ServerUrl = PlasticSourceControlUtils::GetConfigDefaultRepServer();
+			UserName = PlasticSourceControlUtils::GetDefaultUserName();
 		}
 	}
 }
@@ -172,7 +177,7 @@ TSharedRef<FPlasticSourceControlState, ESPMode::ThreadSafe> FPlasticSourceContro
 	else
 	{
 		// cache an unknown state for this item
-		TSharedRef<FPlasticSourceControlState, ESPMode::ThreadSafe> NewState = MakeShared<FPlasticSourceControlState>(FString(InFilename));
+		TSharedRef<FPlasticSourceControlState, ESPMode::ThreadSafe> NewState = MakeShareable(new FPlasticSourceControlState(FString(InFilename)));
 		StateCache.Add(InFilename, NewState);
 		return NewState;
 	}
@@ -203,13 +208,9 @@ void FPlasticSourceControlProvider::HandlePackageSaved(const FString& InPackageF
 
 	// Note: the Editor doesn't ask to refresh the source control status of an asset after it is saved, only *before* (to check that it's possible to save)
 	// So when an asset with no change is saved, update its state in cache to record the fact that the asset is now changed.
-	if (FileState->WorkspaceState == EWorkspaceState::Controlled)
-	{
-		// Note that updating the state in cache isn't enough to refresh the status icon in the Content Browser (since the Editor isn't made aware of the change)
-		// but source control operations are working as expected (eg. "Checkin" and "Revert" are available in the context menu)
-		FileState->WorkspaceState = EWorkspaceState::Changed; // The icon will only appears later when the UI is refreshed (eg switching directory in the Content Browser)
-	}
-	else if (FileState->WorkspaceState == EWorkspaceState::CheckedOutUnchanged)
+	// Note that updating the state in cache isn't enough to refresh the status icon in the Content Browser or the View Changes windows (since the Editor isn't made aware of the change)
+	// so don't trigger the "Changed" status here, it will be handled on the next update status operation.
+	if (FileState->WorkspaceState == EWorkspaceState::CheckedOutUnchanged)
 	{
 		FileState->WorkspaceState = EWorkspaceState::CheckedOutChanged; // In this case the "CheckedOut" icon is already displayed (both states are using the same status icon)
 	}
@@ -222,7 +223,7 @@ FText FPlasticSourceControlProvider::GetStatusText() const
 	Args.Add(TEXT("PluginVersion"), FText::FromString(PluginVersion));
 	Args.Add(TEXT("WorkspacePath"), FText::FromString(PathToWorkspaceRoot));
 	Args.Add(TEXT("WorkspaceName"), FText::FromString(WorkspaceName));
-	Args.Add(TEXT("BranchName"), FText::FromString(BranchName));
+	Args.Add(TEXT("WorkspaceSelector"), FText::FromString(WorkspaceSelector));
 	Args.Add(TEXT("RepositoryName"), FText::FromString(RepositoryName));
 	Args.Add(TEXT("ServerUrl"), FText::FromString(ServerUrl));
 	// Detect special case for a partial checkout (CS:-1 in Gluon mode)!
@@ -256,7 +257,7 @@ FText FPlasticSourceControlProvider::GetStatusText() const
 	}
 	Args.Add(TEXT("ErrorText"), FormattedError);
 
-	return FText::Format(LOCTEXT("PlasticStatusText", "{ErrorText}Unity Version Control (formerly Plastic SCM) {PlasticScmVersion}\t(plugin v{PluginVersion})\nWorkspace: {WorkspaceName}  ({WorkspacePath})\nBranch: {BranchName}@{RepositoryName}@{ServerUrl}\nChangeset: {ChangesetNumber}\nUser: '{UserName}'  {DisplayName}"), Args);
+	return FText::Format(LOCTEXT("PlasticStatusText", "{ErrorText}Unity Version Control (formerly Plastic SCM) {PlasticScmVersion}\t(plugin v{PluginVersion})\nWorkspace: {WorkspaceName}  ({WorkspacePath})\n{WorkspaceSelector}@{RepositoryName}@{ServerUrl}\nChangeset: {ChangesetNumber}\nUser: '{UserName}'  {DisplayName}"), Args);
 }
 
 TMap<ISourceControlProvider::EStatus, FString> FPlasticSourceControlProvider::GetStatus() const
@@ -270,8 +271,8 @@ TMap<ISourceControlProvider::EStatus, FString> FPlasticSourceControlProvider::Ge
 	Result.Add(EStatus::PluginVersion, PluginVersion);
 	Result.Add(EStatus::WorkspacePath, PathToWorkspaceRoot);
 	Result.Add(EStatus::Workspace, WorkspaceName);
-	Result.Add(EStatus::Branch, BranchName);
-	if (IsPartialWorkspace())
+	Result.Add(EStatus::Branch, WorkspaceSelector);
+	if (!IsPartialWorkspace())
 	{
 		Result.Add(EStatus::Changeset, FString::Printf(TEXT("%d"), ChangesetNumber));
 	}
@@ -414,7 +415,7 @@ void FPlasticSourceControlProvider::UnregisterSourceControlStateChanged_Handle(F
 	OnSourceControlStateChanged.Remove(Handle);
 }
 
-ECommandResult::Type FPlasticSourceControlProvider::Execute(const FSourceControlOperationRef& InOperation, FSourceControlChangelistPtr InChangelist, const TArray<FString>& InFiles, EConcurrency::Type InConcurrency, const FSourceControlOperationComplete& InOperationCompleteDelegate)
+	ECommandResult::Type FPlasticSourceControlProvider::Execute(const FSourceControlOperationRef& InOperation, FSourceControlChangelistPtr InChangelist, const TArray<FString>& InFiles, EConcurrency::Type InConcurrency, const FSourceControlOperationComplete& InOperationCompleteDelegate)
 {
 	if (!bWorkspaceFound && !(InOperation->GetName() == "Connect") && !(InOperation->GetName() == "MakeWorkspace"))
 	{
@@ -501,7 +502,7 @@ bool FPlasticSourceControlProvider::UsesCheckout() const
 
 bool FPlasticSourceControlProvider::UsesFileRevisions() const
 {
-	return true;
+	return IsPartialWorkspace();
 }
 
 bool FPlasticSourceControlProvider::UsesSnapshots() const
@@ -635,6 +636,10 @@ void FPlasticSourceControlProvider::UpdateWorkspaceStatus(const class FPlasticSo
 	{
 		BranchName = InCommand.BranchName;
 	}
+	if (!InCommand.WorkspaceSelector.IsEmpty())
+	{
+		WorkspaceSelector = InCommand.WorkspaceSelector;
+	}
 }
 
 void FPlasticSourceControlProvider::Tick()
diff --git a/Engine/Plugins/Developer/PlasticSourceControl/Source/PlasticSourceControl/Private/PlasticSourceControlProvider.h b/Engine/Plugins/Developer/PlasticSourceControl/Source/PlasticSourceControl/Private/PlasticSourceControlProvider.h
--- a/Engine/Plugins/Developer/PlasticSourceControl/Source/PlasticSourceControl/Private/PlasticSourceControlProvider.h
+++ b/Engine/Plugins/Developer/PlasticSourceControl/Source/PlasticSourceControl/Private/PlasticSourceControlProvider.h
@@ -112,8 +112,14 @@ public:
 	{
 		return BranchName;
 	}
-	inline void SetBranchName(const FString& InBranchName)
+	/** Get the Name of the current selector */
+	inline const FString& GetWorkspaceSelector() const
 	{
+		return WorkspaceSelector;
+	}
+	inline void SetWorkspaceSelector(const FString& InWorkspaceSelector, const FString& InBranchName)
+	{
+		WorkspaceSelector = InWorkspaceSelector;
 		BranchName = InBranchName;
 	}
 
@@ -247,6 +253,9 @@ private:
 	/** Name of the current branch */
 	FString BranchName;
 
+	/** Name of the object that the workspace is switched to: usually the branch, sometimes a changeset or a label */
+	FString WorkspaceSelector;
+
 	/** Current Changeset Number */
 	int32 ChangesetNumber = 0;
 
diff --git a/Engine/Plugins/Developer/PlasticSourceControl/Source/PlasticSourceControl/Private/PlasticSourceControlRevision.h b/Engine/Plugins/Developer/PlasticSourceControl/Source/PlasticSourceControl/Private/PlasticSourceControlRevision.h
--- a/Engine/Plugins/Developer/PlasticSourceControl/Source/PlasticSourceControl/Private/PlasticSourceControlRevision.h
+++ b/Engine/Plugins/Developer/PlasticSourceControl/Source/PlasticSourceControl/Private/PlasticSourceControlRevision.h
@@ -79,3 +79,6 @@ public:
 
 /** History composed of the last 100 revisions of the file */
 typedef TArray<TSharedRef<FPlasticSourceControlRevision, ESPMode::ThreadSafe>>	TPlasticSourceControlHistory;
+
+typedef TSharedPtr<FPlasticSourceControlRevision, ESPMode::ThreadSafe> FPlasticSourceControlRevisionPtr;
+typedef TSharedRef<FPlasticSourceControlRevision, ESPMode::ThreadSafe> FPlasticSourceControlRevisionRef;
diff --git a/Engine/Plugins/Developer/PlasticSourceControl/Source/PlasticSourceControl/Private/PlasticSourceControlSettings.cpp b/Engine/Plugins/Developer/PlasticSourceControl/Source/PlasticSourceControl/Private/PlasticSourceControlSettings.cpp
--- a/Engine/Plugins/Developer/PlasticSourceControl/Source/PlasticSourceControl/Private/PlasticSourceControlSettings.cpp
+++ b/Engine/Plugins/Developer/PlasticSourceControl/Source/PlasticSourceControl/Private/PlasticSourceControlSettings.cpp
@@ -59,6 +59,18 @@ void FPlasticSourceControlSettings::SetUpdateStatusOtherBranches(const bool bInU
 	bUpdateStatusOtherBranches = bInUpdateStatusOtherBranches;
 }
 
+bool FPlasticSourceControlSettings::GetViewLocalChanges() const
+{
+	FScopeLock ScopeLock(&CriticalSection);
+	return bViewLocalChanges;
+}
+
+void FPlasticSourceControlSettings::SetViewLocalChanges(const bool bInViewLocalChanges)
+{
+	FScopeLock ScopeLock(&CriticalSection);
+	bViewLocalChanges = bInViewLocalChanges;
+}
+
 bool FPlasticSourceControlSettings::GetEnableVerboseLogs() const
 {
 	FScopeLock ScopeLock(&CriticalSection);
@@ -79,6 +91,7 @@ void FPlasticSourceControlSettings::LoadSettings()
 	GConfig->GetString(*PlasticSettingsConstants::SettingsSection, TEXT("BinaryPath"), BinaryPath, IniFile);
 	GConfig->GetBool(*PlasticSettingsConstants::SettingsSection, TEXT("UpdateStatusAtStartup"), bUpdateStatusAtStartup, IniFile);
 	GConfig->GetBool(*PlasticSettingsConstants::SettingsSection, TEXT("UpdateStatusOtherBranches"), bUpdateStatusOtherBranches, IniFile);
+	GConfig->GetBool(*PlasticSettingsConstants::SettingsSection, TEXT("ViewLocalChanges"), bViewLocalChanges, IniFile);
 	GConfig->GetBool(*PlasticSettingsConstants::SettingsSection, TEXT("EnableVerboseLogs"), bEnableVerboseLogs, IniFile);
 }
 
@@ -89,5 +102,6 @@ void FPlasticSourceControlSettings::SaveSettings() const
 	GConfig->SetString(*PlasticSettingsConstants::SettingsSection, TEXT("BinaryPath"), *BinaryPath, IniFile);
 	GConfig->SetBool(*PlasticSettingsConstants::SettingsSection, TEXT("UpdateStatusAtStartup"), bUpdateStatusAtStartup, IniFile);
 	GConfig->SetBool(*PlasticSettingsConstants::SettingsSection, TEXT("UpdateStatusOtherBranches"), bUpdateStatusOtherBranches, IniFile);
+	GConfig->SetBool(*PlasticSettingsConstants::SettingsSection, TEXT("ViewLocalChanges"), bViewLocalChanges, IniFile);
 	GConfig->SetBool(*PlasticSettingsConstants::SettingsSection, TEXT("EnableVerboseLogs"), bEnableVerboseLogs, IniFile);
 }
diff --git a/Engine/Plugins/Developer/PlasticSourceControl/Source/PlasticSourceControl/Private/PlasticSourceControlSettings.h b/Engine/Plugins/Developer/PlasticSourceControl/Source/PlasticSourceControl/Private/PlasticSourceControlSettings.h
--- a/Engine/Plugins/Developer/PlasticSourceControl/Source/PlasticSourceControl/Private/PlasticSourceControlSettings.h
+++ b/Engine/Plugins/Developer/PlasticSourceControl/Source/PlasticSourceControl/Private/PlasticSourceControlSettings.h
@@ -19,6 +19,10 @@ public:
 	bool GetUpdateStatusOtherBranches() const;
 	void SetUpdateStatusOtherBranches(const bool bInUpdateStatusOtherBranches);
 
+	/** Enable the "View Changes" (changelists) window to also display locally Changed and Private files (can be slow). */
+	bool GetViewLocalChanges() const;
+	void SetViewLocalChanges(const bool bInViewLocalChanges);
+
 	/** Enable LogSourceControl Verbose logs */
 	bool GetEnableVerboseLogs() const;
 	void SetEnableVerboseLogs(const bool bInEnableVerboseLogs);
@@ -45,6 +49,9 @@ private:
 	/** Enable Update status to call "history" to detect recent changesets on other branches (can be slow). */
 	bool bUpdateStatusOtherBranches = false;
 
-	/** Override LogSourceControl verbosity level to Verbose, and back, if not already VeryVerbose */
+	/** Enable the "View Changes" (changelists) window to also display Private and locally Changed assets (even though they cannot be moved to other changelists). */
+	bool bViewLocalChanges = true;
+
+	/** Override LogSourceControl verbosity level to Verbose, and back, if not already VeryVerbose. */
 	bool bEnableVerboseLogs = false;
-};
\ No newline at end of file
+};
diff --git a/Engine/Plugins/Developer/PlasticSourceControl/Source/PlasticSourceControl/Private/PlasticSourceControlShell.cpp b/Engine/Plugins/Developer/PlasticSourceControl/Source/PlasticSourceControl/Private/PlasticSourceControlShell.cpp
--- a/Engine/Plugins/Developer/PlasticSourceControl/Source/PlasticSourceControl/Private/PlasticSourceControlShell.cpp
+++ b/Engine/Plugins/Developer/PlasticSourceControl/Source/PlasticSourceControl/Private/PlasticSourceControlShell.cpp
@@ -80,9 +80,9 @@ static bool _StartBackgroundPlasticShell(const FString& InPathToPlasticBinary, c
 	verify(FPlatformProcess::CreatePipe(ShellOutputPipeRead, ShellOutputPipeWrite, false));	// For reading outputs from cm shell child process
 	verify(FPlatformProcess::CreatePipe(ShellInputPipeRead, ShellInputPipeWrite, true));	// For writing commands to cm shell child process
 
-#if PLATFORM_WINDOWS
+#if !PLATFORM_LINUX // PLATFORM_WINDOWS || PLATFORM_MAC
 	ShellProcessHandle = FPlatformProcess::CreateProc(*InPathToPlasticBinary, *FullCommand, bLaunchDetached, bLaunchHidden, bLaunchReallyHidden, nullptr, 0, *InWorkingDirectory, ShellOutputPipeWrite, ShellInputPipeRead);
-#else // PLATFORM_MAC or PLATFORM_LINUX
+#else // PLATFORM_LINUX
 	// Update working directory
 	char OriginalWorkingDirectory[PATH_MAX];
 	getcwd(OriginalWorkingDirectory, PATH_MAX);
@@ -122,11 +122,13 @@ static void _ExitBackgroundCommandLineShell(const bool bInForceExit = false)
 		{
 			if (bInForceExit)
 			{
+				UE_LOG(LogSourceControl, Verbose, TEXT("_ExitBackgroundCommandLineShell: TerminateProc"));
 				FPlatformProcess::TerminateProc(ShellProcessHandle);
 			}
 			else
 			{
 				// Tell the 'cm shell' to exit
+				UE_LOG(LogSourceControl, Verbose, TEXT("_ExitBackgroundCommandLineShell: exit..."));
 				FPlatformProcess::WritePipe(ShellInputPipeWrite, TEXT("exit"));
 				// And wait up to one second for its termination
 				const double Timeout = 1.0;
@@ -135,7 +137,7 @@ static void _ExitBackgroundCommandLineShell(const bool bInForceExit = false)
 				{
 					if ((FPlatformTime::Seconds() - StartTimestamp) > Timeout)
 					{
-						UE_LOG(LogSourceControl, Warning, TEXT("ExitBackgroundCommandLineShell: cm shell didn't stop gracefully in %lfs."), Timeout);
+						UE_LOG(LogSourceControl, Warning, TEXT("_ExitBackgroundCommandLineShell: cm shell didn't stop gracefully in %lfs."), Timeout);
 						FPlatformProcess::TerminateProc(ShellProcessHandle);
 						break;
 					}
@@ -143,6 +145,10 @@ static void _ExitBackgroundCommandLineShell(const bool bInForceExit = false)
 				}
 			}
 		}
+		else
+		{
+			UE_LOG(LogSourceControl, Verbose, TEXT("_ExitBackgroundCommandLineShell: 'cm shell' already stopped"));
+		}
 		FPlatformProcess::CloseProc(ShellProcessHandle);
 		_CleanupBackgroundCommandLineShell();
 	}
diff --git a/Engine/Plugins/Developer/PlasticSourceControl/Source/PlasticSourceControl/Private/PlasticSourceControlState.cpp b/Engine/Plugins/Developer/PlasticSourceControl/Source/PlasticSourceControl/Private/PlasticSourceControlState.cpp
--- a/Engine/Plugins/Developer/PlasticSourceControl/Source/PlasticSourceControl/Private/PlasticSourceControlState.cpp
+++ b/Engine/Plugins/Developer/PlasticSourceControl/Source/PlasticSourceControl/Private/PlasticSourceControlState.cpp
@@ -2,6 +2,8 @@
 
 #include "PlasticSourceControlState.h"
 #include "PlasticSourceControlProjectSettings.h"
+
+#include "Misc/Paths.h"
 #include "ISourceControlModule.h"
 #include "Styling/AppStyle.h"
 #include "RevisionControlStyle/RevisionControlStyle.h"
@@ -10,28 +12,61 @@
 
 const TCHAR* FPlasticSourceControlState::ToString() const
 {
-	const TCHAR* WorkspaceStateStr = nullptr;
 	switch (WorkspaceState)
 	{
-		case EWorkspaceState::Unknown: WorkspaceStateStr = TEXT("Unknown"); break;
-		case EWorkspaceState::Ignored: WorkspaceStateStr = TEXT("Ignored"); break;
-		case EWorkspaceState::Controlled: WorkspaceStateStr = TEXT("Controlled"); break;
-		case EWorkspaceState::CheckedOutChanged: WorkspaceStateStr = TEXT("CheckedOutChanged"); break;
-		case EWorkspaceState::CheckedOutUnchanged: WorkspaceStateStr = TEXT("CheckedOutUnchanged"); break;
-		case EWorkspaceState::Added: WorkspaceStateStr = TEXT("Added"); break;
-		case EWorkspaceState::Moved: WorkspaceStateStr = TEXT("Moved"); break;
-		case EWorkspaceState::Copied: WorkspaceStateStr = TEXT("Copied"); break;
-		case EWorkspaceState::Replaced: WorkspaceStateStr = TEXT("Replaced"); break;
-		case EWorkspaceState::Deleted: WorkspaceStateStr = TEXT("Deleted"); break;
-		case EWorkspaceState::LocallyDeleted: WorkspaceStateStr = TEXT("LocallyDeleted"); break;
-		case EWorkspaceState::Changed: WorkspaceStateStr = TEXT("Changed"); break;
-		case EWorkspaceState::Conflicted: WorkspaceStateStr = TEXT("Conflicted"); break;
-		case EWorkspaceState::Private: WorkspaceStateStr = TEXT("Private"); break;
-		default: WorkspaceStateStr = TEXT("???"); break;
+		case EWorkspaceState::Ignored: return TEXT("Ignored");
+		case EWorkspaceState::Controlled: return TEXT("Controlled");
+		case EWorkspaceState::CheckedOutChanged: return TEXT("Checked-out (changed)");
+		case EWorkspaceState::CheckedOutUnchanged: return TEXT("Checked-out (unchanged)");
+		case EWorkspaceState::Added: return TEXT("Added");
+		case EWorkspaceState::Moved: return TEXT("Moved");
+		case EWorkspaceState::Copied: return TEXT("Copied");
+		case EWorkspaceState::Replaced: return TEXT("Replaced");
+		case EWorkspaceState::Deleted: return TEXT("Deleted");
+		case EWorkspaceState::LocallyDeleted: return TEXT("LocallyDeleted");
+		case EWorkspaceState::Changed: return TEXT("Changed");
+		case EWorkspaceState::Conflicted: return TEXT("Conflicted");
+		case EWorkspaceState::Private: return TEXT("Private");
+		case EWorkspaceState::Unknown:
+		default: return TEXT("Unknown");
 	}
-	return WorkspaceStateStr;
 }
 
+FText FPlasticSourceControlState::ToText() const
+{
+	static const FText Ignored = LOCTEXT("Ignored", "Ignored");
+	static const FText Controlled = LOCTEXT("Controlled", "Controlled");
+	static const FText CheckedOutChanged = LOCTEXT("CheckedOutChanged", "Checked-out (changed)");
+	static const FText CheckedOutUnchanged = LOCTEXT("CheckedOutUnchanged", "Checked-out (unchanged)");
+	static const FText Added = LOCTEXT("Added", "Added");
+	static const FText Moved = LOCTEXT("Moved", "Moved");
+	static const FText Copied = LOCTEXT("Copied", "Copied");
+	static const FText Replaced = LOCTEXT("Replaced", "Replaced");
+	static const FText Deleted = LOCTEXT("Deleted", "Deleted");
+	static const FText LocallyDeleted = LOCTEXT("LocallyDeleted", "LocallyDeleted");
+	static const FText Changed = LOCTEXT("Changed", "Changed");
+	static const FText Conflicted = LOCTEXT("Conflicted", "Conflicted");
+	static const FText Private = LOCTEXT("Private", "Private");
+	static const FText Unknown = LOCTEXT("Unknown", "Unknown");
+	switch (WorkspaceState)
+	{
+		case EWorkspaceState::Ignored: return Ignored;
+		case EWorkspaceState::Controlled: return Controlled;
+		case EWorkspaceState::CheckedOutChanged: return CheckedOutChanged;
+		case EWorkspaceState::CheckedOutUnchanged: return CheckedOutUnchanged;
+		case EWorkspaceState::Added: return Added;
+		case EWorkspaceState::Moved: return Moved;
+		case EWorkspaceState::Copied: return Copied;
+		case EWorkspaceState::Replaced: return Replaced;
+		case EWorkspaceState::Deleted: return Deleted;
+		case EWorkspaceState::LocallyDeleted: return LocallyDeleted;
+		case EWorkspaceState::Changed: return Changed;
+		case EWorkspaceState::Conflicted: return Conflicted;
+		case EWorkspaceState::Private: return Private;
+		case EWorkspaceState::Unknown:
+		default: return Unknown;
+	}
+}
 
 int32 FPlasticSourceControlState::GetHistorySize() const
 {
@@ -320,7 +355,7 @@ bool FPlasticSourceControlState::CanCheckout() const
 
 	const bool bCanCheckout  =    (WorkspaceState == EWorkspaceState::Controlled	// In source control, Unmodified
 								|| WorkspaceState == EWorkspaceState::Changed)		// In source control, but not checked-out
-								&& !IsCheckedOutOther()	// Is not already checked-out elsewhere
+							//	&& !IsCheckedOutOther()	// Note: with multi-destination branches it is sometimes possible to checkout an asset already checked-out elsewhere. Let the user try.
 								&& IsCurrent();			// Is up to date (at the revision of the repo)
 
 	if (!IsUnknown())
@@ -352,6 +387,13 @@ bool FPlasticSourceControlState::IsCheckedOut() const
 	}
 }
 
+bool FPlasticSourceControlState::IsPendingChanges() const
+{
+	return WorkspaceState != EWorkspaceState::Unknown
+		&& WorkspaceState != EWorkspaceState::Ignored
+		&& WorkspaceState != EWorkspaceState::Controlled;
+}
+
 bool FPlasticSourceControlState::IsCheckedOutImplementation() const
 {
 	const bool bIsCheckedOut = WorkspaceState == EWorkspaceState::CheckedOutChanged
diff --git a/Engine/Plugins/Developer/PlasticSourceControl/Source/PlasticSourceControl/Private/PlasticSourceControlState.h b/Engine/Plugins/Developer/PlasticSourceControl/Source/PlasticSourceControl/Private/PlasticSourceControlState.h
--- a/Engine/Plugins/Developer/PlasticSourceControl/Source/PlasticSourceControl/Private/PlasticSourceControlState.h
+++ b/Engine/Plugins/Developer/PlasticSourceControl/Source/PlasticSourceControl/Private/PlasticSourceControlState.h
@@ -50,6 +50,19 @@ public:
 		Move(MoveTemp(InState));
 	}
 
+	// Comparison operator designed to detect and report only meaningful changes to the Editor, mainly for the purpose of updating Content Browser overlay icons
+	bool operator==(const FPlasticSourceControlState& InState) const
+	{
+		return (WorkspaceState == InState.WorkspaceState)
+			&& (LockedBy == InState.LockedBy)
+			&& (RetainedBy == InState.RetainedBy)
+			&& (IsCurrent() == InState.IsCurrent());
+	}
+	bool operator!=(const FPlasticSourceControlState& InState) const
+	{
+		return !(*this == InState);
+	}
+
 	const FPlasticSourceControlState& operator=(FPlasticSourceControlState&& InState)
 	{
 		Move(MoveTemp(InState));
@@ -58,7 +71,10 @@ public:
 
 	void Move(FPlasticSourceControlState&& InState)
 	{
-		History = MoveTemp(InState.History);
+		if (InState.History.Num() > 0)
+		{
+			History = MoveTemp(InState.History);
+		}
 		LocalFilename = MoveTemp(InState.LocalFilename);
 		WorkspaceState = InState.WorkspaceState;
 		PendingResolveInfo = MoveTemp(InState.PendingResolveInfo);
@@ -96,6 +112,13 @@ public:
 	// debug log utility
 	const TCHAR* ToString() const;
 
+	FText ToText() const;
+
+	void PopulateSearchString(TArray<FString>& OutStrings) const
+	{
+		OutStrings.Emplace(LocalFilename);
+	}
+
 	/** ISourceControlState interface */
 	virtual int32 GetHistorySize() const override;
 	virtual TSharedPtr<class ISourceControlRevision, ESPMode::ThreadSafe> GetHistoryItem(int32 HistoryIndex) const override;
@@ -133,6 +156,7 @@ public:
 	virtual bool IsConflicted() const override;
 	virtual bool CanRevert() const override;
 
+	bool IsPendingChanges() const;
 	bool IsCheckedOutImplementation() const;
 	bool IsLocked() const;
 	bool IsRetainedInOtherBranch() const;
@@ -204,3 +228,6 @@ public:
 	/** The change list of the last modification */
 	int32 HeadChangeList;
 };
+
+typedef TSharedRef<FPlasticSourceControlState, ESPMode::ThreadSafe> FPlasticSourceControlStateRef;
+typedef TSharedPtr<FPlasticSourceControlState, ESPMode::ThreadSafe> FPlasticSourceControlStatePtr;
