PR #12084: [Source Control] Update Plastic SCM provider to 1.11.0 (from 1.9.0)

diff --git a/Engine/Plugins/Developer/PlasticSourceControl/Source/PlasticSourceControl/Private/PlasticSourceControlUtils.cpp b/Engine/Plugins/Developer/PlasticSourceControl/Source/PlasticSourceControl/Private/PlasticSourceControlUtils.cpp
--- a/Engine/Plugins/Developer/PlasticSourceControl/Source/PlasticSourceControl/Private/PlasticSourceControlUtils.cpp
+++ b/Engine/Plugins/Developer/PlasticSourceControl/Source/PlasticSourceControl/Private/PlasticSourceControlUtils.cpp
@@ -3,7 +3,9 @@
 #include "PlasticSourceControlUtils.h"
 
 #include "PlasticSourceControlBranch.h"
+#include "PlasticSourceControlChangeset.h"
 #include "PlasticSourceControlCommand.h"
+#include "PlasticSourceControlLock.h"
 #include "PlasticSourceControlModule.h"
 #include "PlasticSourceControlParsers.h"
 #include "PlasticSourceControlProjectSettings.h"
@@ -21,6 +23,12 @@
 
 #include "PlasticSourceControlChangelistState.h"
 
+#if PLATFORM_WINDOWS
+#include "Windows/AllowWindowsPlatformTypes.h"
+#include "Windows/WindowsPlatformMisc.h"
+#undef GetUserName
+#endif
+
 namespace PlasticSourceControlUtils
 {
 
@@ -40,8 +48,18 @@ bool RunCommand(const FString& InCommand, const TArray<FString>& InParameters, c
 
 	const bool bResult = PlasticSourceControlShell::RunCommand(InCommand, InParameters, InFiles, Results, Errors);
 
-	Results.ParseIntoArray(OutResults, PlasticSourceControlShell::pchDelim, true);
-	Errors.ParseIntoArray(OutErrorMessages, PlasticSourceControlShell::pchDelim, true);
+	if (!Results.IsEmpty())
+	{
+		TArray<FString> ParsedResults;
+		Results.ParseIntoArray(ParsedResults, PlasticSourceControlShell::pchDelim, true);
+		OutResults.Append(MoveTemp(ParsedResults));
+	}
+	if (!Errors.IsEmpty())
+	{
+		TArray<FString> ParsedErrors;
+		Errors.ParseIntoArray(ParsedErrors, PlasticSourceControlShell::pchDelim, true);
+		OutErrorMessages.Append(MoveTemp(ParsedErrors));
+	}
 
 	return bResult;
 }
@@ -55,6 +73,105 @@ FString FindPlasticBinaryPath()
 #endif
 }
 
+static FString FindDesktopApplicationPath()
+{
+	FString DesktopAppPath;
+
+#if PLATFORM_WINDOWS
+	// On Windows, use the registry to find the install location
+	FString InstallLocation = TEXT("C:/Program Files/PlasticSCM5");
+	if (FWindowsPlatformMisc::QueryRegKey(HKEY_LOCAL_MACHINE, TEXT("SOFTWARE\\Unity Software Inc.\\Unity DevOps Version Control"), TEXT("Location"), InstallLocation))
+	{
+		FPaths::NormalizeDirectoryName(InstallLocation);
+	}
+
+	const TCHAR* PlasticExe = TEXT("client/plastic.exe");
+	const TCHAR* GluonExe = TEXT("client/gluon.exe");
+	DesktopAppPath = FPaths::Combine(InstallLocation, FPlasticSourceControlModule::Get().GetProvider().IsPartialWorkspace() ? GluonExe : PlasticExe);
+#elif PLATFORM_MAC
+	const TCHAR* PlasticExe = TEXT("/Applications/PlasticSCM.app/Contents/MacOS/macplasticx");
+	const TCHAR* GluonExe = TEXT("/Applications/Gluon.app/Contents/MacOS/macgluonx");
+	DesktopAppPath = FPlasticSourceControlModule::Get().GetProvider().IsPartialWorkspace() ? GluonExe : PlasticExe;
+#elif PLATFORM_LINUX
+	const TCHAR* PlasticExe = TEXT("/usr/bin/plasticgui");
+	const TCHAR* GluonExe = TEXT("/usr/bin/gluon");
+	DesktopAppPath = FPlasticSourceControlModule::Get().GetProvider().IsPartialWorkspace() ? GluonExe : PlasticExe;
+#endif
+
+	return DesktopAppPath;
+}
+
+static bool OpenDesktopApplication(const FString& InCommandLineArguments)
+{
+	const FString DesktopAppPath = FindDesktopApplicationPath();
+
+	UE_LOG(LogSourceControl, Log, TEXT("Opening the Desktop application (%s %s)"), *DesktopAppPath, *InCommandLineArguments);
+
+	FProcHandle Proc = FPlatformProcess::CreateProc(*DesktopAppPath, *InCommandLineArguments, true, false, false, nullptr, 0, nullptr, nullptr, nullptr);
+	if (!Proc.IsValid())
+	{
+		UE_LOG(LogSourceControl, Error, TEXT("Opening the Desktop application (%s %s) failed."), *DesktopAppPath, *InCommandLineArguments);
+		FPlatformProcess::CloseProc(Proc);
+		return false;
+	}
+
+	return true;
+}
+
+bool OpenDesktopApplication(const bool bInBranchExplorer)
+{
+	const FString CommandLineArguments = FString::Printf(TEXT("--wk=\"%s\" %s"),
+		*FPlasticSourceControlModule::Get().GetProvider().GetPathToWorkspaceRoot(),
+		bInBranchExplorer ? TEXT("--view=BranchExplorerView") : TEXT(""));
+
+	return OpenDesktopApplication(CommandLineArguments);
+}
+
+static FString GetFullSpec(const int32 InChangesetId)
+{
+	const FPlasticSourceControlProvider& Provider = FPlasticSourceControlModule::Get().GetProvider();
+	return FString::Printf(TEXT("cs:%d@%s@%s"), InChangesetId, *Provider.GetRepositoryName(), *Provider.GetServerUrl());
+}
+
+static FString GetFullSpec(const FString& InBranchName)
+{
+	const FPlasticSourceControlProvider& Provider = FPlasticSourceControlModule::Get().GetProvider();
+	return FString::Printf(TEXT("br:%s@%s@%s"), *InBranchName, *Provider.GetRepositoryName(), *Provider.GetServerUrl());
+}
+
+bool OpenDesktopApplicationForDiff(const int32 InChangesetId)
+{
+	const FString CommandLineArguments = FString::Printf(TEXT("--diffchangeset=\"%s\""),
+		*GetFullSpec(InChangesetId));
+
+	return OpenDesktopApplication(CommandLineArguments);
+}
+
+bool OpenDesktopApplicationForDiff(const int32 InChangesetIdSrc, const int32 InChangesetIdDst)
+{
+	const FString CommandLineArguments = FString::Printf(TEXT("--diffchangesetsrc=\"%s\" --diffchangesetdst=\"%s\""),
+		*GetFullSpec(InChangesetIdSrc), *GetFullSpec(InChangesetIdDst));
+
+	return OpenDesktopApplication(CommandLineArguments);
+}
+
+bool OpenDesktopApplicationForDiff(const FString& InBranchName)
+{
+	const FString CommandLineArguments = FString::Printf(TEXT("--diffbranch=\"%s\""),
+		*GetFullSpec(InBranchName));
+
+	return OpenDesktopApplication(CommandLineArguments);
+}
+
+void OpenLockRulesInCloudDashboard(const FString& InOrganizationName)
+{
+	const FString OrganizationLockRulesURL = FString::Printf(
+		TEXT("https://dashboard.unity3d.com/devops/organizations/default/plastic-scm/organizations/%s/lock-rules"),
+		*InOrganizationName
+	);
+	FPlatformProcess::LaunchURL(*OrganizationLockRulesURL, NULL, NULL);
+}
+
 // Find the root of the workspace, looking from the provided path and upward in its parent directories.
 bool GetWorkspacePath(const FString& InPath, FString& OutWorkspaceRoot)
 {
@@ -136,15 +253,35 @@ FString GetConfigDefaultRepServer()
 	return ServerUrl;
 }
 
-void GetUserName(FString& OutUserName)
+FString GetDefaultUserName()
 {
+	FString UserName;
 	TArray<FString> Results;
 	TArray<FString> ErrorMessages;
 	const bool bResult = RunCommand(TEXT("whoami"), TArray<FString>(), TArray<FString>(), Results, ErrorMessages);
 	if (bResult && Results.Num() > 0)
 	{
-		OutUserName = MoveTemp(Results[0]);
+		UserName = MoveTemp(Results[0]);
 	}
+
+	return UserName;
+}
+
+FString GetProfileUserName(const FString& InServerUrl)
+{
+	FString UserName;
+	TArray<FString> Results;
+	TArray<FString> Parameters;
+	TArray<FString> ErrorMessages;
+	Parameters.Add("list");
+	Parameters.Add(TEXT("--format=\"{server};{user}\""));
+	bool bResult = RunCommand(TEXT("profile"), Parameters, TArray<FString>(), Results, ErrorMessages);
+	if (bResult)
+	{
+		bResult = PlasticSourceControlParsers::ParseProfileInfo(Results, InServerUrl, UserName);
+	}
+
+	return UserName;
 }
 
 bool GetWorkspaceName(const FString& InWorkspaceRoot, FString& OutWorkspaceName, TArray<FString>& OutErrorMessages)
@@ -168,22 +305,38 @@ bool GetWorkspaceName(const FString& InWorkspaceRoot, FString& OutWorkspaceName,
 	return bResult;
 }
 
-bool GetWorkspaceInfo(FString& OutBranchName, FString& OutRepositoryName, FString& OutServerUrl, TArray<FString>& OutErrorMessages)
+bool GetWorkspaceInfo(FString& OutWorkspaceSelector, FString& OutBranchName, FString& OutRepositoryName, FString& OutServerUrl, TArray<FString>& OutErrorMessages)
 {
 	TArray<FString> Results;
 	bool bResult = RunCommand(TEXT("workspaceinfo"), TArray<FString>(), TArray<FString>(), Results, OutErrorMessages);
 	if (bResult)
 	{
-		bResult = PlasticSourceControlParsers::ParseWorkspaceInfo(Results, OutBranchName, OutRepositoryName, OutServerUrl);
+		bResult = PlasticSourceControlParsers::ParseWorkspaceInfo(Results, OutWorkspaceSelector, OutBranchName, OutRepositoryName, OutServerUrl);
+	}
+
+	return bResult;
+}
+
+bool GetChangesetNumber(int32& OutChangesetNumber, TArray<FString>& OutErrorMessages)
+{
+	TArray<FString> Results;
+	TArray<FString> Parameters;
+	Parameters.Add(TEXT("--header"));
+	Parameters.Add(TEXT("--machinereadable"));
+	Parameters.Add(TEXT("--fieldseparator=\"") FILE_STATUS_SEPARATOR TEXT("\""));
+	bool bResult = RunCommand(TEXT("status"), Parameters, TArray<FString>(), Results, OutErrorMessages);
+	if (bResult)
+	{
+		bResult = PlasticSourceControlParsers::GetChangesetFromWorkspaceStatus(Results, OutChangesetNumber);
 	}
 
 	return bResult;
 }
 
-bool RunCheckConnection(FString& OutBranchName, FString& OutRepositoryName, FString& OutServerUrl, TArray<FString>& OutInfoMessages, TArray<FString>& OutErrorMessages)
+bool RunCheckConnection(FString& OutWorkspaceSelector, FString& OutBranchName, FString& OutRepositoryName, FString& OutServerUrl, TArray<FString>& OutInfoMessages, TArray<FString>& OutErrorMessages)
 {
 	TArray<FString> Parameters;
-	if (PlasticSourceControlUtils::GetWorkspaceInfo(OutBranchName, OutRepositoryName, OutServerUrl, OutErrorMessages))
+	if (PlasticSourceControlUtils::GetWorkspaceInfo(OutWorkspaceSelector, OutBranchName, OutRepositoryName, OutServerUrl, OutErrorMessages))
 	{
 		Parameters.Add(FString::Printf(TEXT("--server=%s"), *OutServerUrl));
 	}
@@ -222,9 +375,8 @@ FString UserNameToDisplayName(const FString& InUserName)
  * @param[out]	OutErrorMessages	Error messages from the "status" command
  * @param[out]	OutStates			States of files for witch the status has been gathered (distinct than InFiles in case of a "directory status")
  * @param[out]	OutChangeset		The current Changeset Number
- * @param[out]	OutBranchName		Name of the current checked-out branch
  */
-static bool RunStatus(const FString& InDir, TArray<FString>&& InFiles, const EStatusSearchType InSearchType, TArray<FString>& OutErrorMessages, TArray<FPlasticSourceControlState>& OutStates, int32& OutChangeset, FString& OutBranchName)
+static bool RunStatus(const FString& InDir, TArray<FString>&& InFiles, const EStatusSearchType InSearchType, TArray<FString>& OutErrorMessages, TArray<FPlasticSourceControlState>& OutStates, int32& OutChangeset)
 {
 	TRACE_CPUPROFILER_EVENT_SCOPE(PlasticSourceControlUtils::RunStatus);
 
@@ -233,29 +385,25 @@ static bool RunStatus(const FString& InDir, TArray<FString>&& InFiles, const ESt
 	TArray<FString> Parameters;
 	Parameters.Add(TEXT("--machinereadable"));
 	Parameters.Add(TEXT("--fieldseparator=\"") FILE_STATUS_SEPARATOR TEXT("\""));
+	Parameters.Add(TEXT("--controlledchanged"));
 	if (InSearchType == EStatusSearchType::All)
 	{
 		// NOTE: don't use "--all" to avoid searching for --localmoved since it's the most time consuming (beside --changed)
 		// and its not used by the plugin (matching similarities doesn't seem to work with .uasset files)
-		Parameters.Add(TEXT("--controlledchanged"));
 		// TODO: add a user settings to avoid searching for --changed and --localdeleted, to work like Perforce on big projects,
 		// provided that the user understands the consequences (they won't see assets modified locally without a proper checkout)
 		Parameters.Add(TEXT("--changed"));
 		Parameters.Add(TEXT("--localdeleted"));
 		Parameters.Add(TEXT("--private"));
 		Parameters.Add(TEXT("--ignored"));
-
-		// If the version of cm is recent enough use the new --iscochanged for "CO+CH" status
-		const FPlasticSourceControlProvider& Provider = FPlasticSourceControlModule::Get().GetProvider();
-		const bool bUsesCheckedOutChanged = Provider.GetPlasticScmVersion() >= PlasticSourceControlVersions::StatusIsCheckedOutChanged;
-		if (bUsesCheckedOutChanged)
-		{
-			Parameters.Add(TEXT("--iscochanged"));
-		}
 	}
-	else if (InSearchType == EStatusSearchType::ControlledOnly)
+
+	// If the version of cm is recent enough use the new --iscochanged for "CO+CH" status
+	const FPlasticSourceControlProvider& Provider = FPlasticSourceControlModule::Get().GetProvider();
+	const bool bUsesCheckedOutChanged = Provider.GetPlasticScmVersion() >= PlasticSourceControlVersions::StatusIsCheckedOutChanged;
+	if (bUsesCheckedOutChanged)
 	{
-		Parameters.Add(TEXT("--controlledchanged"));
+		Parameters.Add(TEXT("--iscochanged"));
 	}
 
 	// "cm status" only operate on one path (file or directory) at a time, so use one common path for multiple files in a directory
@@ -309,9 +457,59 @@ static bool RunStatus(const FString& InDir, TArray<FString>&& InFiles, const ESt
 	return bResult;
 }
 
-bool RunListSmartLocks(const FString& InRepository, TMap<FString, PlasticSourceControlParsers::FSmartLockInfoParser>& OutSmartLocks)
+// Cache Locks with a Timestamp, and an InvalidateCachedLocks() function
+class FLocksCache
 {
-	TRACE_CPUPROFILER_EVENT_SCOPE(PlasticSourceControlUtils::RunListSmartLocks);
+public:
+	void Reset()
+	{
+		FScopeLock Lock(&CriticalSection);
+		Locks.Reset();
+		Timestamp = FDateTime();
+	}
+
+	void SetLocks(const TArray<FPlasticSourceControlLockRef>& InLocks)
+	{
+		FScopeLock Lock(&CriticalSection);
+		Locks = InLocks;
+		Timestamp = FDateTime::Now();
+	}
+
+	bool GetLocks(TArray<FPlasticSourceControlLockRef>& OutLocks)
+	{
+		FScopeLock Lock(&CriticalSection);
+		const FTimespan ElapsedTime = FDateTime::Now() - Timestamp;
+		if (ElapsedTime.GetTotalSeconds() < 60.0)
+		{
+			OutLocks = Locks;
+			return true;
+		}
+		return false;
+	}
+
+private:
+	TArray<FPlasticSourceControlLockRef> Locks;
+	FDateTime Timestamp;
+	FCriticalSection CriticalSection;
+};
+
+static FLocksCache LocksCacheForAllDestBranches;
+static FLocksCache LocksCacheForWorkingBranch;
+
+void InvalidateLocksCache()
+{
+	LocksCacheForAllDestBranches.Reset();
+	LocksCacheForWorkingBranch.Reset();
+}
+
+bool RunListLocks(const FPlasticSourceControlProvider& InProvider, const bool bInForAllDestBranches, TArray<FPlasticSourceControlLockRef>& OutLocks)
+{
+	TRACE_CPUPROFILER_EVENT_SCOPE(PlasticSourceControlUtils::RunListLocks);
+
+	FLocksCache& LocksCache = bInForAllDestBranches ? LocksCacheForAllDestBranches : LocksCacheForWorkingBranch;
+
+	if (LocksCache.GetLocks(OutLocks))
+		return true;
 
 	TArray<FString> Results;
 	TArray<FString> ErrorMessages;
@@ -319,25 +517,76 @@ bool RunListSmartLocks(const FString& InRepository, TMap<FString, PlasticSourceC
 	Parameters.Add(TEXT("list"));
 	Parameters.Add(TEXT("--machinereadable"));
 	Parameters.Add(TEXT("--smartlocks"));
+	Parameters.Add(FString::Printf(TEXT("--repository=%s"), *InProvider.GetRepositoryName()));
 	Parameters.Add(TEXT("--anystatus"));
 	Parameters.Add(TEXT("--fieldseparator=\"") FILE_STATUS_SEPARATOR TEXT("\""));
+	// NOTE: --dateformat was added to smartlocks a couple of releases later in version 11.0.16.8133
 	Parameters.Add(TEXT("--dateformat=yyyy-MM-ddTHH:mm:ss"));
-	bool bResult = RunCommand(TEXT("lock"), Parameters, TArray<FString>(), Results, ErrorMessages);
+	// For displaying Locks as a status overlay icon in the Content Browser, restricts the Locks to only those applying to the current branch so there can be only one and never any ambiguity
+	if (!bInForAllDestBranches && (InProvider.GetPlasticScmVersion() >= PlasticSourceControlVersions::WorkingBranch))
+	{
+		// Note: here is one of the rare places where we need to use a branch name, not a workspace selector
+		Parameters.Add(FString::Printf(TEXT("--workingbranch=%s"), *InProvider.GetBranchName()));
+	}
+	const bool bResult = RunCommand(TEXT("lock"), Parameters, TArray<FString>(), Results, ErrorMessages);
 
 	if (bResult)
 	{
+		OutLocks.Reserve(Results.Num());
 		for (int32 IdxResult = 0; IdxResult < Results.Num(); IdxResult++)
 		{
 			const FString& Result = Results[IdxResult];
-			PlasticSourceControlParsers::FSmartLockInfoParser SmartLockInfoParser(Result);
-			if (SmartLockInfoParser.Repository == InRepository)
+			FPlasticSourceControlLock&& Lock = PlasticSourceControlParsers::ParseLockInfo(Result);
+			OutLocks.Add(MakeShareable(new FPlasticSourceControlLock(Lock)));
+		}
+
+		LocksCache.SetLocks(OutLocks);
+	}
+
+	return bResult;
+}
+
+TArray<FPlasticSourceControlLockRef> GetLocksForWorkingBranch(const FPlasticSourceControlProvider& InProvider, const TArray<FString>& InFiles)
+{
+	TArray<FPlasticSourceControlLockRef> Locks;
+
+	// Only get locks for the current working branch
+	const bool bInForAllDestBranches = false;
+	RunListLocks(InProvider, bInForAllDestBranches, Locks);
+
+	TArray<FPlasticSourceControlLockRef> MatchingLocks;
+	MatchingLocks.Reserve(InFiles.Num());
+
+	// Only return locks for the specified files
+	for (const FString& File : InFiles)
+	{
+		for (const FPlasticSourceControlLockRef& Lock : Locks)
+		{
+			if (File.EndsWith(Lock->Path))
 			{
-				OutSmartLocks.Add(SmartLockInfoParser.Filename, SmartLockInfoParser);
+				MatchingLocks.Add(Lock);
+				break;
 			}
 		}
 	}
 
-	return bResult;
+	return MatchingLocks;
+}
+
+TArray<FString> LocksToFileNames(const FString InWorkspaceRoot, const TArray<FPlasticSourceControlLockRef>& InLocks)
+{
+	TArray<FString> Files;
+
+	// Note: remove the slash '/' from the end of the Workspace root to Combine it with server paths also starting with a slash
+	const FString& WorkspaceRoot = InWorkspaceRoot[InWorkspaceRoot.Len() - 1] == TEXT('/') ? InWorkspaceRoot.LeftChop(1) : InWorkspaceRoot;
+
+	Files.Reserve(InLocks.Num());
+	for (const FPlasticSourceControlLockRef& Lock : InLocks)
+	{
+		Files.AddUnique(FPaths::Combine(WorkspaceRoot, Lock->Path));
+	}
+
+	return Files;
 }
 
 /**
@@ -527,7 +776,7 @@ struct FFilesInCommonDir
 };
 
 // Run a batch of Plastic "status" and "fileinfo" commands to update status of given files and directories.
-bool RunUpdateStatus(const TArray<FString>& InFiles, const EStatusSearchType InSearchType, const bool bInUpdateHistory, TArray<FString>& OutErrorMessages, TArray<FPlasticSourceControlState>& OutStates, int32& OutChangeset, FString& OutBranchName)
+bool RunUpdateStatus(const TArray<FString>& InFiles, const EStatusSearchType InSearchType, const bool bInUpdateHistory, TArray<FString>& OutErrorMessages, TArray<FPlasticSourceControlState>& OutStates, int32& OutChangeset)
 {
 	bool bResults = true;
 
@@ -592,7 +841,7 @@ bool RunUpdateStatus(const TArray<FString>& InFiles, const EStatusSearchType InS
 		// This should be an edge case (typically the uproject file) .
 		if (!bDirFound)
 		{
-			FString Path = FPaths::GetPath(File) + TEXT('/');
+			const FString Path = FPaths::GetPath(File) + TEXT('/');
 			FFilesInCommonDir* ExistingGroup = GroupOfFiles.Find(Path);
 			if (ExistingGroup != nullptr)
 			{
@@ -600,7 +849,7 @@ bool RunUpdateStatus(const TArray<FString>& InFiles, const EStatusSearchType InS
 			}
 			else
 			{
-				GroupOfFiles.Add(Path, { MoveTemp(Path), {File} });
+				GroupOfFiles.Add(Path, { Path, {File} });
 			}
 		}
 	}
@@ -622,7 +871,7 @@ bool RunUpdateStatus(const TArray<FString>& InFiles, const EStatusSearchType InS
 		// Run a "status" command on the directory to get workspace file states.
 		// (ie. Changed, CheckedOut, Copied, Replaced, Added, Private, Ignored, Deleted, LocallyDeleted, Moved, LocallyMoved)
 		TArray<FPlasticSourceControlState> States;
-		const bool bGroupOk = RunStatus(Group.Value.CommonDir, MoveTemp(Group.Value.Files), InSearchType, OutErrorMessages, States, OutChangeset, OutBranchName);
+		const bool bGroupOk = RunStatus(Group.Value.CommonDir, MoveTemp(Group.Value.Files), InSearchType, OutErrorMessages, States, OutChangeset);
 		if (!bGroupOk)
 		{
 			bResults = false;
@@ -662,23 +911,6 @@ bool RunGetFile(const FString& InRevSpec, const FString& InDumpFileName)
 	return bResult;
 }
 
-// Convert a file state to a string ala Perforce, see also ParseShelveFileStatus()
-FString FileStateToAction(const EWorkspaceState InState)
-{
-	switch (InState)
-	{
-	case EWorkspaceState::Added:
-		return TEXT("add");
-	case EWorkspaceState::Deleted:
-		return TEXT("delete");
-	case EWorkspaceState::Moved:
-		return TEXT("branch");
-	case EWorkspaceState::CheckedOutChanged:
-	default:
-		return TEXT("edit");
-	}
-}
-
 // Run a Plastic "history" command and parse it's XML result.
 bool RunGetHistory(const bool bInUpdateHistory, TArray<FPlasticSourceControlState>& InOutStates, TArray<FString>& OutErrorMessages)
 {
@@ -693,7 +925,8 @@ bool RunGetHistory(const bool bInUpdateHistory, TArray<FPlasticSourceControlStat
 	{
 		Parameters.Add(TEXT("--moveddeleted"));
 	}
-	Parameters.Add(TEXT("--xml"));
+	const FScopedTempFile HistoryResultFile(TEXT("History-"), TEXT(".xml"));
+	Parameters.Add(FString::Printf(TEXT("--xml=\"%s\""), *HistoryResultFile.GetFilename()));
 	Parameters.Add(TEXT("--encoding=\"utf-8\""));
 	const FPlasticSourceControlProvider& Provider = FPlasticSourceControlModule::Get().GetProvider();
 	if (Provider.GetPlasticScmVersion() >= PlasticSourceControlVersions::NewHistoryLimit)
@@ -728,7 +961,7 @@ bool RunGetHistory(const bool bInUpdateHistory, TArray<FPlasticSourceControlStat
 		bResult = RunCommand(TEXT("history"), Parameters, Files, Results, Errors);
 		if (bResult)
 		{
-			bResult = PlasticSourceControlParsers::ParseHistoryResults(bInUpdateHistory, Results, InOutStates);
+			bResult = PlasticSourceControlParsers::ParseHistoryResults(bInUpdateHistory, HistoryResultFile.GetFilename(), InOutStates);
 		}
 		if (!Errors.IsEmpty())
 		{
@@ -740,7 +973,7 @@ bool RunGetHistory(const bool bInUpdateHistory, TArray<FPlasticSourceControlStat
 }
 
 // Run a Plastic "update" command to sync the workspace and parse its XML results.
-bool RunUpdate(const TArray<FString>& InFiles, const bool bInIsPartialWorkspace, TArray<FString>& OutUpdatedFiles, TArray<FString>& OutErrorMessages)
+bool RunUpdate(const TArray<FString>& InFiles, const bool bInIsPartialWorkspace, const FString& InChangesetId, TArray<FString>& OutUpdatedFiles, TArray<FString>& OutErrorMessages)
 {
 	bool bResult = false;
 
@@ -749,12 +982,20 @@ bool RunUpdate(const TArray<FString>& InFiles, const bool bInIsPartialWorkspace,
 	// Detect special case for a partial checkout (CS:-1 in Gluon mode)!
 	if (!bInIsPartialWorkspace)
 	{
-		const FScopedTempFile UpdateResultFile;
+		const FScopedTempFile UpdateResultFile(TEXT("Update-"), TEXT(".xml"));
 		TArray<FString> InfoMessages;
+		if (!InChangesetId.IsEmpty())
+		{
+			Parameters.Add(FString::Printf(TEXT("--changeset=%s"), *InChangesetId));
+		}
+		else
+		{
+			Parameters.Add(TEXT("--last"));
+		}
+		Parameters.Add(TEXT("--dontmerge"));
+		Parameters.Add(TEXT("--noinput"));
 		Parameters.Add(FString::Printf(TEXT("--xml=\"%s\""), *UpdateResultFile.GetFilename()));
 		Parameters.Add(TEXT("--encoding=\"utf-8\""));
-		Parameters.Add(TEXT("--last"));
-		Parameters.Add(TEXT("--dontmerge"));
 		bResult = PlasticSourceControlUtils::RunCommand(TEXT("update"), Parameters, TArray<FString>(), InfoMessages, OutErrorMessages);
 		if (bResult)
 		{
@@ -769,6 +1010,10 @@ bool RunUpdate(const TArray<FString>& InFiles, const bool bInIsPartialWorkspace,
 	else
 	{
 		TArray<FString> Results;
+		if (!InChangesetId.IsEmpty())
+		{
+			Parameters.Add(FString::Printf(TEXT("--changeset=%s"), *InChangesetId));
+		}
 		Parameters.Add(TEXT("--report"));
 		Parameters.Add(TEXT("--machinereadable"));
 		bResult = PlasticSourceControlUtils::RunCommand(TEXT("partial update"), Parameters, InFiles, Results, OutErrorMessages);
@@ -786,18 +1031,37 @@ bool RunGetChangelists(TArray<FPlasticSourceControlChangelistState>& OutChangeli
 {
 	TRACE_CPUPROFILER_EVENT_SCOPE(PlasticSourceControlUtils::RunGetChangelists);
 
+	const FPlasticSourceControlProvider& Provider = FPlasticSourceControlModule::Get().GetProvider();
+
 	FString Results;
 	FString Errors;
+	const FScopedTempFile GetChangelistFile(TEXT("StatusChangelist-"), TEXT(".xml"));
 	TArray<FString> Parameters;
 	Parameters.Add(TEXT("--changelists"));
 	Parameters.Add(TEXT("--controlledchanged"));
+	if (Provider.AccessSettings().GetViewLocalChanges())
+	{
+		// NOTE: don't use "--all" to avoid searching for --localmoved since it's the most time consuming (beside --changed)
+		// and its not used by the plugin (matching similarities doesn't seem to work with .uasset files)
+		Parameters.Add(TEXT("--changed"));
+		Parameters.Add(TEXT("--localdeleted"));
+		Parameters.Add(TEXT("--private"));
+	}
+
+	// If the version of cm is recent enough use the new --iscochanged for "CO+CH" status
+	const bool bUsesCheckedOutChanged = Provider.GetPlasticScmVersion() >= PlasticSourceControlVersions::StatusIsCheckedOutChanged;
+	if (bUsesCheckedOutChanged)
+	{
+		Parameters.Add(TEXT("--iscochanged"));
+	}
+
 	Parameters.Add(TEXT("--noheader"));
-	Parameters.Add(TEXT("--xml"));
+	Parameters.Add(FString::Printf(TEXT("--xml=\"%s\""), *GetChangelistFile.GetFilename()));
 	Parameters.Add(TEXT("--encoding=\"utf-8\""));
 	bool bResult = RunCommand(TEXT("status"), Parameters, TArray<FString>(), Results, Errors);
 	if (bResult)
 	{
-		bResult = PlasticSourceControlParsers::ParseChangelistsResults(Results, OutChangelistsStates, OutCLFilesStates);
+		bResult = PlasticSourceControlParsers::ParseChangelistsResults(GetChangelistFile.GetFilename(), OutChangelistsStates, OutCLFilesStates);
 	}
 	if (!Errors.IsEmpty())
 	{
@@ -879,7 +1143,7 @@ bool RunGetShelves(TArray<FPlasticSourceControlChangelistState>& InOutChangelist
 	FString Errors;
 	TArray<FString> Parameters;
 	Parameters.Add(TEXT("\"shelves where owner = 'me'\""));
-	Parameters.Add(TEXT("--xml"));
+	Parameters.Add(TEXT("--xml")); // TODO convert to -xml=<FScopedTempFile>
 	Parameters.Add(TEXT("--encoding=\"utf-8\""));
 	bCommandSuccessful = PlasticSourceControlUtils::RunCommand(TEXT("find"), Parameters, TArray<FString>(), Results, Errors);
 	if (bCommandSuccessful)
@@ -934,7 +1198,7 @@ bool RunGetShelve(const int32 InShelveId, FString& OutComment, FDateTime& OutDat
 	FString Errors;
 	TArray<FString> Parameters;
 	Parameters.Add(FString::Printf(TEXT("\"shelves where ShelveId = %d\""), InShelveId));
-	Parameters.Add(TEXT("--xml"));
+	Parameters.Add(TEXT("--xml")); // TODO convert to -xml=<FScopedTempFile>
 	Parameters.Add(TEXT("--encoding=\"utf-8\""));
 	bCommandSuccessful = PlasticSourceControlUtils::RunCommand(TEXT("find"), Parameters, TArray<FString>(), Results, Errors);
 	if (bCommandSuccessful)
@@ -953,31 +1217,82 @@ bool RunGetShelve(const int32 InShelveId, FString& OutComment, FDateTime& OutDat
 	return bCommandSuccessful;
 }
 
+bool RunGetChangesets(const FDateTime& InFromDate, TArray<FPlasticSourceControlChangesetRef>& OutChangesets, TArray<FString>& OutErrorMessages)
+{
+	bool bCommandSuccessful = false;
+
+	const FScopedTempFile ChangesetResultFile(TEXT("FindChangeset-"), TEXT(".xml"));
+	TArray<FString> Results;
+	TArray<FString> Errors;
+	TArray<FString> Parameters;
+	Parameters.Add(TEXT("changesets"));
+	if (InFromDate != FDateTime())
+	{
+		Parameters.Add(FString::Printf(TEXT("\"where date >= '%d/%d/%d'\""), InFromDate.GetYear(), InFromDate.GetMonth(), InFromDate.GetDay()));
+	}
+	Parameters.Add(TEXT("order by ChangesetId desc"));
+	Parameters.Add(FString::Printf(TEXT("--xml=\"%s\""), *ChangesetResultFile.GetFilename()));
+	Parameters.Add(TEXT("--encoding=\"utf-8\""));
+	bCommandSuccessful = PlasticSourceControlUtils::RunCommand(TEXT("find"), Parameters, TArray<FString>(), Results, Errors);
+	if (bCommandSuccessful && FPaths::FileExists(ChangesetResultFile.GetFilename()))
+	{
+		bCommandSuccessful = PlasticSourceControlParsers::ParseChangesetsResults(ChangesetResultFile.GetFilename(), OutChangesets);
+	}
+	if (Errors.Num() > 0)
+	{
+		OutErrorMessages.Append(MoveTemp(Errors));
+	}
+
+	return bCommandSuccessful;
+}
+
+bool RunGetChangesetFiles(const FPlasticSourceControlChangesetRef& InChangeset, TArray<FPlasticSourceControlStateRef>& OutFiles, TArray<FString>& OutErrorMessages)
+{
+	bool bCommandSuccessful = false;
+
+	const FScopedTempFile LogResultFile(TEXT("Log-"), TEXT(".xml"));
+	FString Results;
+	FString Errors;
+	TArray<FString> Parameters;
+	Parameters.Add(FString::Printf(TEXT("cs:%d"), InChangeset->ChangesetId));
+	Parameters.Add(FString::Printf(TEXT("--xml=\"%s\""), *LogResultFile.GetFilename()));
+	Parameters.Add(TEXT("--encoding=\"utf-8\""));
+	bCommandSuccessful = PlasticSourceControlUtils::RunCommand(TEXT("log"), Parameters, TArray<FString>(), Results, Errors);
+	if (bCommandSuccessful && FPaths::FileExists(LogResultFile.GetFilename()))
+	{
+		bCommandSuccessful = PlasticSourceControlParsers::ParseLogResults(LogResultFile.GetFilename(), InChangeset, OutFiles);
+	}
+	if (!Errors.IsEmpty())
+	{
+		OutErrorMessages.Add(MoveTemp(Errors));
+	}
+
+	return bCommandSuccessful;
+}
+
 bool RunGetBranches(const FDateTime& InFromDate, TArray<FPlasticSourceControlBranchRef>& OutBranches, TArray<FString>& OutErrorMessages)
 {
 	bool bCommandSuccessful;
 
+	const FScopedTempFile BranchResultFile(TEXT("FindBranches-"), TEXT(".xml"));
 	FString Results;
 	FString Errors;
 	TArray<FString> Parameters;
+	Parameters.Add(TEXT("branches"));
 	if (InFromDate != FDateTime())
 	{
 		// Find branches created since this date or containing changes dating from or after this date
-		Parameters.Add(FString::Printf(TEXT("\"branches where date >= '%d/%d/%d'\" or changesets >= '%d/%d/%d'"),
+		Parameters.Add(FString::Printf(TEXT("\"where date >= '%d/%d/%d'\" or changesets >= '%d/%d/%d'"),
 			InFromDate.GetYear(), InFromDate.GetMonth(), InFromDate.GetDay(),
 			InFromDate.GetYear(), InFromDate.GetMonth(), InFromDate.GetDay()
 		));
 	}
-	else
-	{
-		Parameters.Add(TEXT("\"branches\""));
-	}
-	Parameters.Add(TEXT("--xml"));
+	Parameters.Add(FString::Printf(TEXT("--xml=\"%s\""), *BranchResultFile.GetFilename()));
 	Parameters.Add(TEXT("--encoding=\"utf-8\""));
 	bCommandSuccessful = PlasticSourceControlUtils::RunCommand(TEXT("find"), Parameters, TArray<FString>(), Results, Errors);
 	if (bCommandSuccessful)
 	{
-		bCommandSuccessful = PlasticSourceControlParsers::ParseBranchesResults(Results, OutBranches);
+		bCommandSuccessful = PlasticSourceControlParsers::ParseBranchesResults(BranchResultFile.GetFilename(), OutBranches);
 	}
 	if (!Errors.IsEmpty())
 	{
@@ -987,22 +1302,45 @@ bool RunGetBranches(const FDateTime& InFromDate, TArray<FPlasticSourceControlBra
 	return bCommandSuccessful;
 }
 
-bool RunSwitchToBranch(const FString& InBranchName, TArray<FString>& OutUpdatedFiles, TArray<FString>& OutErrorMessages)
+bool RunSwitch(const FString& InBranchName, const int32 InChangesetId, const bool bInIsPartialWorkspace, TArray<FString>& OutUpdatedFiles, TArray<FString>& OutErrorMessages)
 {
 	bool bResult = false;
 
-	const FScopedTempFile SwitchResultFile;
+	const FScopedTempFile SwitchResultFile(TEXT("Switch-"), TEXT(".xml"));
 	TArray<FString> InfoMessages;
 	TArray<FString> Parameters;
-	Parameters.Add(FString::Printf(TEXT("--xml=\"%s\""), *SwitchResultFile.GetFilename()));
-	Parameters.Add(TEXT("--encoding=\"utf-8\""));
-	Parameters.Add(FString::Printf(TEXT("\"br:%s\""), *InBranchName));
-	bResult = PlasticSourceControlUtils::RunCommand(TEXT("switch"), Parameters, TArray<FString>(), InfoMessages, OutErrorMessages);
-	if (bResult)
+	if (InChangesetId != ISourceControlState::INVALID_REVISION)
 	{
-		// Load and parse the result of the update command
-		FString Results;
-		if (FFileHelper::LoadFileToString(Results, *SwitchResultFile.GetFilename()))
+		// NOTE: not supported by Gluon/partial workspaces
+		Parameters.Add(FString::Printf(TEXT("cs:%d"), InChangesetId));
+	}
+	else
+	{
+		Parameters.Add(FString::Printf(TEXT("\"br:%s\""), *InBranchName));
+	}
+	Parameters.Add(TEXT("--noinput"));
+	// Detect special case for a partial checkout (CS:-1 in Gluon mode)!
+	if (!bInIsPartialWorkspace)
+	{
+		Parameters.Add(FString::Printf(TEXT("--xml=\"%s\""), *SwitchResultFile.GetFilename()));
+		Parameters.Add(TEXT("--encoding=\"utf-8\""));
+		bResult = PlasticSourceControlUtils::RunCommand(TEXT("switch"), Parameters, TArray<FString>(), InfoMessages, OutErrorMessages);
+		if (bResult)
+		{
+			// Load and parse the result of the update command
+			FString Results;
+			if (FFileHelper::LoadFileToString(Results, *SwitchResultFile.GetFilename()))
+			{
+				bResult = PlasticSourceControlParsers::ParseUpdateResults(Results, OutUpdatedFiles);
+			}
+		}
+	}
+	else
+	{
+		TArray<FString> Results;
+		Parameters.Add(TEXT("--report"));
+		bResult = PlasticSourceControlUtils::RunCommand(TEXT("partial switch"), Parameters, TArray<FString>(), Results, OutErrorMessages);
+		if (bResult)
 		{
 			bResult = PlasticSourceControlParsers::ParseUpdateResults(Results, OutUpdatedFiles);
 		}
@@ -1015,7 +1353,7 @@ bool RunMergeBranch(const FString& InBranchName, TArray<FString>& OutUpdatedFile
 {
 	bool bResult = false;
 
-	const FScopedTempFile MergeResultFile;
+	const FScopedTempFile MergeResultFile(TEXT("Merge-"), TEXT(".xml"));
 	TArray<FString> InfoMessages;
 	TArray<FString> Parameters;
 	Parameters.Add(FString::Printf(TEXT("--xml=\"%s\""), *MergeResultFile.GetFilename()));
@@ -1080,14 +1418,20 @@ bool UpdateCachedStates(TArray<FPlasticSourceControlState>&& InStates)
 	FPlasticSourceControlProvider& Provider = FPlasticSourceControlModule::Get().GetProvider();
 	const FDateTime Now = FDateTime::Now();
 
+	bool bUpdatedStates = false;
 	for (auto&& InState : InStates)
 	{
 		TSharedRef<FPlasticSourceControlState, ESPMode::ThreadSafe> State = Provider.GetStateInternal(InState.LocalFilename);
+		// Only report that the cache was updated if the state changed in a meaningful way, useful to the Editor
+		if (*State != InState)
+		{
+			bUpdatedStates = true;
+		}
 		*State = MoveTemp(InState);
 		State->TimeStamp = Now;
 	}
 
-	return (InStates.Num() > 0);
+	return bUpdatedStates;
 }
 
 void RemoveRedundantErrors(FPlasticSourceControlCommand& InCommand, const FString& InFilter)
diff --git a/Engine/Plugins/Developer/PlasticSourceControl/Source/PlasticSourceControl/Private/PlasticSourceControlUtils.h b/Engine/Plugins/Developer/PlasticSourceControl/Source/PlasticSourceControl/Private/PlasticSourceControlUtils.h
--- a/Engine/Plugins/Developer/PlasticSourceControl/Source/PlasticSourceControl/Private/PlasticSourceControlUtils.h
+++ b/Engine/Plugins/Developer/PlasticSourceControl/Source/PlasticSourceControl/Private/PlasticSourceControlUtils.h
@@ -8,14 +8,13 @@
 
 class FPlasticSourceControlChangelistState;
 class FPlasticSourceControlCommand;
+class FPlasticSourceControlProvider;
 class FPlasticSourceControlState;
 struct FSoftwareVersion;
 typedef TSharedRef<class FPlasticSourceControlBranch, ESPMode::ThreadSafe> FPlasticSourceControlBranchRef;
-
-namespace PlasticSourceControlParsers
-{
-class FSmartLockInfoParser;
-} // namespace PlasticSourceControlParsers
+typedef TSharedRef<class FPlasticSourceControlChangeset, ESPMode::ThreadSafe> FPlasticSourceControlChangesetRef;
+typedef TSharedRef<class FPlasticSourceControlLock, ESPMode::ThreadSafe> FPlasticSourceControlLockRef;
+typedef TSharedRef<class FPlasticSourceControlState, ESPMode::ThreadSafe> FPlasticSourceControlStateRef;
 
 enum class EWorkspaceState;
 
@@ -51,6 +50,42 @@ bool RunCommand(const FString& InCommand, const TArray<FString>& InParameters, c
  */
 FString FindPlasticBinaryPath();
 
+/**
+ * Open the Desktop Application.
+ *
+ * @param	bInBranchExplorer	Should it open the Branch Explorer instead of the Workspace Explorer (false by default)?
+ */
+bool OpenDesktopApplication(const bool bInBranchExplorer = false);
+
+/**
+ * Open the Desktop Application for diffing a Changeset.
+ *
+ * @param	InChangesetId		Changeset to diff
+ */
+bool OpenDesktopApplicationForDiff(const int32 InChangesetId);
+
+/**
+ * Open the Desktop Application for diffing a couple of Changesets.
+ *
+ * @param	InChangesetIdSrc	Source (Left) Changeset to diff
+ * @param	InChangesetIdDst	Destination (Right) Changeset to diff
+ */
+bool OpenDesktopApplicationForDiff(const int32 InChangesetIdSrc, const int32 InChangesetIdDst);
+
+/**
+ * Open the Desktop Application for diffing a whole Branch.
+ *
+ * @param	InBranchName		Name of the Branch to diff
+ */
+bool OpenDesktopApplicationForDiff(const FString& InBranchName);
+
+/**
+ * Open the Unity Cloud Dashboard on the page to show and manage Lock Rules.
+ *
+ * @param	InOrganizationName	Name of the organization to use, from the server URL
+ */
+void OpenLockRulesInCloudDashboard(const FString& InOrganizationName);
+
 /**
  * Find the root of the Plastic workspace, looking from the GameDir and upward in its parent directories
  * @param InPathToGameDir		The path to the Game Directory
@@ -86,14 +121,21 @@ bool GetConfigSetFilesAsReadOnly();
 FString GetConfigDefaultRepServer();
 
 /**
- * Get Unity Version Control current user
- * @param	OutUserName			Name of the Unity Version Control user configured globally
+ * Get Unity Version Control user for the default server
+ * @returns	Name of the Unity Version Control user configured globally
  */
-void GetUserName(FString& OutUserName);
+FString GetDefaultUserName();
+
+/**
+ * Get Unity Version Control user for the specified server
+ * @param	InServerUrl		Name of the specified server
+ * @returns	Name of the Unity Version Control user for the specified server
+ */
+FString GetProfileUserName(const FString& InServerUrl);
 
 /**
  * Get workspace name
- * @param	InWorkspaceRoot		The workspace from where to run the command - usually the Game directory
+ * @param	InWorkspaceRoot		The workspace from where to run the command - typically the Project path
  * @param	OutWorkspaceName	Name of the current workspace
  * @param	OutErrorMessages	Any errors (from StdErr) as an array per-line
 */
@@ -101,24 +143,35 @@ bool GetWorkspaceName(const FString& InWorkspaceRoot, FString& OutWorkspaceName,
 
 /**
  * Get workspace info: the current branch, repository name, and server URL
- * 
+ *
+ * @param	OutWorkspaceSelector	Name of the current branch, changeset or label depending on the workspace selector
  * @param	OutBranchName		Name of the current branch
  * @param	OutRepositoryName	Name of the repository of the current workspace
  * @param	OutServerUrl		URL/Port of the server of the repository
  * @param	OutErrorMessages	Any errors (from StdErr) as an array per-line
  */
-bool GetWorkspaceInfo(FString& OutBranchName, FString& OutRepositoryName, FString& OutServerUrl, TArray<FString>& OutErrorMessages);
+bool GetWorkspaceInfo(FString& OutWorkspaceSelector, FString& OutBranchName, FString& OutRepositoryName, FString& OutServerUrl, TArray<FString>& OutErrorMessages);
+
 
 /**
- * Get workspace info and check the connection to the server
+ * Get the current changeset number from the workspace (-1 for a partial workspace)
  *
- * @param	OutBranchName		Name of the current branch
- * @param	OutRepositoryName	Name of the repository of the current workspace
- * @param	OutServerUrl		URL/Port of the server of the repository
- * @param	OutInfoMessages		Result of the connection test
+ * @param	OutChangesetNumber	Current changeset number
  * @param	OutErrorMessages	Any errors (from StdErr) as an array per-line
  */
-bool RunCheckConnection(FString& OutBranchName, FString& OutRepositoryName, FString& OutServerUrl, TArray<FString>& OutInfoMessages, TArray<FString>& OutErrorMessages);
+bool GetChangesetNumber(int32& OutChangesetNumber, TArray<FString>& OutErrorMessages);
+
+/**
+ * Get workspace info and check the connection to the server
+ *
+ * @param	OutWorkspaceSelector	Name of the current branch, changeset or label depending on the workspace selector
+ * @param	OutBranchName			Name of the current branch when available
+ * @param	OutRepositoryName		Name of the repository of the current workspace
+ * @param	OutServerUrl			URL/Port of the server of the repository
+ * @param	OutInfoMessages			Result of the connection test
+ * @param	OutErrorMessages		Any errors (from StdErr) as an array per-line
+ */
+bool RunCheckConnection(FString& OutWorkspaceSelector, FString& OutBranchName, FString& OutRepositoryName, FString& OutServerUrl, TArray<FString>& OutInfoMessages, TArray<FString>& OutErrorMessages);
 
 /**
  * Use the Project Settings to replace Unity Version Control full username/e-mail by a shorter version for display.
@@ -129,14 +182,37 @@ bool RunCheckConnection(FString& OutBranchName, FString& OutRepositoryName, FStr
  */
 FString UserNameToDisplayName(const FString& InUserName);
 
+/**
+ * Invalidate the cache of locks so that the next call to RunListLocks() will not use it and actually run the cm lock list command
+ */
+void InvalidateLocksCache();
+
 /**
  * Run a Plastic "lock list" command and parse it.
  *
- * @param	InRepository		The repository to ask for the locks
- * @param	OutSmartLocks		The list of smart locks
+ * @param	InProvider				The source control provider to get the repository and current branch to ask the locks for
+ * @param   bInForAllDestBranches	Retrieve locks for all destination branches, or restrict them to only those applying to the working branch
+ * @param	OutLocks				The list of locks
  * @returns true if the command succeeded and returned no errors
  */
-bool RunListSmartLocks(const FString& InRepository, TMap<FString, PlasticSourceControlParsers::FSmartLockInfoParser>& OutSmartLocks);
+bool RunListLocks(const FPlasticSourceControlProvider& InProvider, const bool bInForAllDestBranches, TArray<FPlasticSourceControlLockRef>& OutLocks);
+
+/**
+ * Get locks applying to the working branch for the specified files.
+ *
+ * @param	InProvider			The source control provider to get the repository and current branch to ask the locks for
+ * @param	InFiles				The files to be operated on (server paths)
+ * @return	OutLocks			The list of corresponding locks if any
+ */
+TArray<FPlasticSourceControlLockRef> GetLocksForWorkingBranch(const FPlasticSourceControlProvider& InProvider, const TArray<FString>& InFiles);
+
+/**
+ * Get the list of filenames from the list of locks
+ * @param	InWorkspaceRoot		The workspace from where to run the command - typically the Project path
+ * @param	InLocks				Locks to get the file names for
+ * @return	List of absolute filenames
+*/
+TArray<FString> LocksToFileNames(const FString InWorkspaceRoot, const TArray<FPlasticSourceControlLockRef>& InLocks);
 
 // Specify the "search type" for the "status" command
 enum class EStatusSearchType
@@ -154,10 +230,9 @@ enum class EStatusSearchType
  * @param	OutErrorMessages	Any errors (from StdErr) as an array per-line
  * @param	OutStates			States of the files
  * @param	OutChangeset		The current Changeset Number
- * @param	OutBranchName		Name of the current checked-out branch
  * @returns true if the command succeeded and returned no errors
  */
-bool RunUpdateStatus(const TArray<FString>& InFiles, const EStatusSearchType InSearchType, const bool bInUpdateHistory, TArray<FString>& OutErrorMessages, TArray<FPlasticSourceControlState>& OutStates, int32& OutChangeset, FString& OutBranchName);
+bool RunUpdateStatus(const TArray<FString>& InFiles, const EStatusSearchType InSearchType, const bool bInUpdateHistory, TArray<FString>& OutErrorMessages, TArray<FPlasticSourceControlState>& OutStates, int32& OutChangeset);
 
 /**
  * Run a Plastic "cat" command to dump the binary content of a revision into a file.
@@ -182,10 +257,11 @@ bool RunGetHistory(const bool bInUpdateHistory, TArray<FPlasticSourceControlStat
  *
  * @param	InFiles					The files or paths to sync
  * @param	bInIsPartialWorkspace	Whether running on a partial/gluon or regular/full workspace
+ * @param	InChangesetId			The optional changeset to sync to (leave empty to sync to the latest in the branch)
  * @param	OutUpdatedFiles			The files that where updated
  * @param	OutErrorMessages		Any errors (from StdErr) as an array per-line
  */
-bool RunUpdate(const TArray<FString>& InFiles, const bool bInIsPartialWorkspace, TArray<FString>& OutUpdatedFiles, TArray<FString>& OutErrorMessages);
+bool RunUpdate(const TArray<FString>& InFiles, const bool bInIsPartialWorkspace, const FString& InChangesetId, TArray<FString>& OutUpdatedFiles, TArray<FString>& OutErrorMessages);
 
 /**
  * Run a Plastic "status --changelist --xml" and parse its XML result.
@@ -222,6 +298,24 @@ bool RunGetShelve(const int32 InShelveId, FString& OutComment, FDateTime& OutDat
 */
 void AddShelvedFileToChangelist(FPlasticSourceControlChangelistState& InOutChangelistsState, FString&& InFilename, EWorkspaceState InShelveStatus, FString&& InMovedFrom);
 
+/**
+ * Run find "changesets where date >= 'YYYY-MM-DD'" and parse the results.
+ * @param	InFromDate				The date to search from
+ * @param	OutChangesets			The list of changesets, without their files
+ * @param	OutErrorMessages		Any errors (from StdErr) as an array per-line
+ *
+ * @see RunGetChangesetFiles() below used to populated a specific changeset with its list of files
+ */
+bool RunGetChangesets(const FDateTime& InFromDate, TArray<FPlasticSourceControlChangesetRef>& OutChangesets, TArray<FString>& OutErrorMessages);
+
+/**
+ * Run "log cs:<ChangesetId> --xml" and parse the results to populate the files from the specified changeset.
+ * @param	InChangeset				The changeset to get the files changed
+ * @param	OutFiles				The files changed in the specified changeset
+ * @param	OutErrorMessages		Any errors (from StdErr) as an array per-line
+ */
+bool RunGetChangesetFiles(const FPlasticSourceControlChangesetRef& InChangeset, TArray<FPlasticSourceControlStateRef>& OutFiles, TArray<FString>& OutErrorMessages);
+
 /**
  * Run find "branches where date >= 'YYYY-MM-DD' or changesets >= 'YYYY-MM-DD'" and parse the results.
  * @param	InFromDate				The date to search from
@@ -232,11 +326,13 @@ bool RunGetBranches(const FDateTime& InFromDate, TArray<FPlasticSourceControlBra
 
 /**
  * Run switch br:/name and parse the results.
- * @param	InBranchName			The name of the branch to switch the workspace to
+ * @param	InBranchName			The name of the branch to switch the workspace to (optional, only used if no InChangesetId)
+ * @param	InChangesetId			The name of the changeset to switch the workspace to (optional, overrides InBranchName if set)
+ * @param	bInIsPartialWorkspace	Whether running on a partial/gluon or regular/full workspace
  * @param	OutUpdatedFiles			The files that where updated
  * @param	OutErrorMessages		Any errors (from StdErr) as an array per-line
  */
-bool RunSwitchToBranch(const FString& InBranchName, TArray<FString>& OutUpdatedFiles, TArray<FString>& OutErrorMessages);
+bool RunSwitch(const FString& InBranchName, const int32 InChangesetId, const bool bInIsPartialWorkspace, TArray<FString>& OutUpdatedFiles, TArray<FString>& OutErrorMessages);
 
 /**
  * Run merge br:/name and parse the results.
@@ -275,7 +371,7 @@ bool RunDeleteBranches(const TArray<FString>& InBranchNames, TArray<FString>& Ou
  */
 bool UpdateCachedStates(TArray<FPlasticSourceControlState>&& InStates);
 
-/** 
+/**
  * Remove redundant errors (that contain a particular string) and also
  * update the commands success status if all errors were removed.
  */
diff --git a/Engine/Plugins/Developer/PlasticSourceControl/Source/PlasticSourceControl/Private/PlasticSourceControlVersions.h b/Engine/Plugins/Developer/PlasticSourceControl/Source/PlasticSourceControl/Private/PlasticSourceControlVersions.h
--- a/Engine/Plugins/Developer/PlasticSourceControl/Source/PlasticSourceControl/Private/PlasticSourceControlVersions.h
+++ b/Engine/Plugins/Developer/PlasticSourceControl/Source/PlasticSourceControl/Private/PlasticSourceControlVersions.h
@@ -44,6 +44,12 @@ namespace PlasticSourceControlVersions
 
 	// 11.0.16.8101 Introducing Smart Locks
 	// https://plasticscm.com/download/releasenotes/11.0.16.8101 (2023/07/06)
-	static const FSoftwareVersion SmartLocks(TEXT("11.0.16.8101"));
+	// 11.0.16.8133 add support for non-localized --datetimeformat
+	// https://plasticscm.com/download/releasenotes/11.0.16.8133 (2023/08/03)
+	static const FSoftwareVersion SmartLocks(TEXT("11.0.16.8133"));
+
+	// 11.0.16.8445 lock list new --workingbranch option to only list the locks that apply to the destination branch of the specified branch.
+	// https://plasticscm.com/download/releasenotes/11.0.16.8445 (2024/02/22)
+	static const FSoftwareVersion WorkingBranch(TEXT("11.0.16.8445"));
 
 } // namespace PlasticSourceControlVersions
diff --git a/Engine/Plugins/Developer/PlasticSourceControl/Source/PlasticSourceControl/Private/PlasticSourceControlWorkspaceCreation.cpp b/Engine/Plugins/Developer/PlasticSourceControl/Source/PlasticSourceControl/Private/PlasticSourceControlWorkspaceCreation.cpp
--- a/Engine/Plugins/Developer/PlasticSourceControl/Source/PlasticSourceControl/Private/PlasticSourceControlWorkspaceCreation.cpp
+++ b/Engine/Plugins/Developer/PlasticSourceControl/Source/PlasticSourceControl/Private/PlasticSourceControlWorkspaceCreation.cpp
@@ -47,8 +47,11 @@ void FPlasticSourceControlWorkspaceCreation::OnMakeWorkspaceOperationComplete(co
 {
 	OnSourceControlOperationComplete(InOperation, InResult);
 
-	// Launch the next asynchronous operation
-	LaunchMarkForAddOperation();
+	if (InResult == ECommandResult::Succeeded)
+	{
+		// Launch the next asynchronous operation
+		LaunchMarkForAddOperation();
+	}
 }
 
 /// 2. Add all project files to Source Control (.uproject, Config/, Content/, Source/ files and ignore.conf if any)
@@ -84,8 +87,11 @@ void FPlasticSourceControlWorkspaceCreation::OnMarkForAddOperationComplete(const
 {
 	OnSourceControlOperationComplete(InOperation, InResult);
 
-	// Launch the next asynchronous operation
-	LaunchCheckInOperation();
+	if (InResult == ECommandResult::Succeeded)
+	{
+		// Launch the next asynchronous operation
+		LaunchCheckInOperation();
+	}
 }
 
 /// 3. Launch an asynchronous "CheckIn" operation and start another ongoing notification
diff --git a/Engine/Plugins/Developer/PlasticSourceControl/Source/PlasticSourceControl/Private/SPlasticSourceControlBranchRow.cpp b/Engine/Plugins/Developer/PlasticSourceControl/Source/PlasticSourceControl/Private/SPlasticSourceControlBranchRow.cpp
--- a/Engine/Plugins/Developer/PlasticSourceControl/Source/PlasticSourceControl/Private/SPlasticSourceControlBranchRow.cpp
+++ b/Engine/Plugins/Developer/PlasticSourceControl/Source/PlasticSourceControl/Private/SPlasticSourceControlBranchRow.cpp
@@ -9,7 +9,7 @@
 
 #include "Styling/AppStyle.h"
 
-#define LOCTEXT_NAMESPACE "PlasticSourceControlWindow"
+#define LOCTEXT_NAMESPACE "PlasticSourceControlBranchesWindow"
 
 FName PlasticSourceControlBranchesListViewColumn::Name::Id() { return TEXT("Name"); }
 FText PlasticSourceControlBranchesListViewColumn::Name::GetDisplayText() { return LOCTEXT("Name_Column", "Name"); }
@@ -84,6 +84,7 @@ TSharedRef<SWidget> SPlasticSourceControlBranchRow::GenerateWidgetForColumn(cons
 	}
 	else if (InColumnId == PlasticSourceControlBranchesListViewColumn::Comment::Id())
 	{
+		// Make each comment fit on a single line to preserve the table layout
 		FString CommentOnOneLine = BranchToVisualize->Comment;
 		CommentOnOneLine.ReplaceCharInline(TEXT('\n'), TEXT(' '), ESearchCase::CaseSensitive);
 
@@ -101,4 +102,4 @@ TSharedRef<SWidget> SPlasticSourceControlBranchRow::GenerateWidgetForColumn(cons
 	}
 }
 
-#undef LOCTEXT_NAMESPACE
\ No newline at end of file
+#undef LOCTEXT_NAMESPACE
diff --git a/Engine/Plugins/Developer/PlasticSourceControl/Source/PlasticSourceControl/Private/SPlasticSourceControlBranchesWidget.cpp b/Engine/Plugins/Developer/PlasticSourceControl/Source/PlasticSourceControl/Private/SPlasticSourceControlBranchesWidget.cpp
--- a/Engine/Plugins/Developer/PlasticSourceControl/Source/PlasticSourceControl/Private/SPlasticSourceControlBranchesWidget.cpp
+++ b/Engine/Plugins/Developer/PlasticSourceControl/Source/PlasticSourceControl/Private/SPlasticSourceControlBranchesWidget.cpp
@@ -6,6 +6,7 @@
 #include "PlasticSourceControlOperations.h"
 #include "PlasticSourceControlProjectSettings.h"
 #include "PlasticSourceControlBranch.h"
+#include "PlasticSourceControlUtils.h"
 #include "PlasticSourceControlVersions.h"
 #include "SPlasticSourceControlBranchRow.h"
 #include "SPlasticSourceControlCreateBranch.h"
@@ -25,28 +26,34 @@
 #include "Styling/AppStyle.h"
 #include "Framework/Application/SlateApplication.h"
 #include "Framework/Docking/TabManager.h"
+#include "Widgets/Images/SImage.h"
+#include "Widgets/Input/SButton.h"
 #include "Widgets/Input/SComboButton.h"
 #include "Widgets/Input/SSearchBox.h"
 #include "Widgets/Layout/SSpacer.h"
-#include "Widgets/SBoxPanel.h"
-#include "Widgets/SWindow.h"
 #include "Widgets/Text/STextBlock.h"
 #include "Widgets/Views/SHeaderRow.h"
+#include "Widgets/SBoxPanel.h"
+#include "Widgets/SWindow.h"
 
-#define LOCTEXT_NAMESPACE "PlasticSourceControlWindow"
+#define LOCTEXT_NAMESPACE "PlasticSourceControlBranchesWindow"
 
 void SPlasticSourceControlBranchesWidget::Construct(const FArguments& InArgs)
 {
 	ISourceControlModule::Get().RegisterProviderChanged(FSourceControlProviderChanged::FDelegate::CreateSP(this, &SPlasticSourceControlBranchesWidget::OnSourceControlProviderChanged));
+	// register for any source control change to detect any change of branch from the Changesets window
+	SourceControlStateChangedDelegateHandle = ISourceControlModule::Get().GetProvider().RegisterSourceControlStateChanged_Handle(FSourceControlStateChanged::FDelegate::CreateSP(this, &SPlasticSourceControlBranchesWidget::HandleSourceControlStateChanged));
 
-	CurrentBranchName = FPlasticSourceControlModule::Get().GetProvider().GetBranchName();
+	WorkspaceSelector = FPlasticSourceControlModule::Get().GetProvider().GetWorkspaceSelector();
 
 	SearchTextFilter = MakeShared<TTextFilter<const FPlasticSourceControlBranch&>>(TTextFilter<const FPlasticSourceControlBranch&>::FItemToStringArray::CreateSP(this, &SPlasticSourceControlBranchesWidget::PopulateItemSearchStrings));
 	SearchTextFilter->OnChanged().AddSP(this, &SPlasticSourceControlBranchesWidget::OnRefreshUI);
 
 	FromDateInDaysValues.Add(TPair<int32, FText>(7, FText::FromString(TEXT("Last week"))));
+	FromDateInDaysValues.Add(TPair<int32, FText>(15, FText::FromString(TEXT("Last 15 days"))));
 	FromDateInDaysValues.Add(TPair<int32, FText>(30, FText::FromString(TEXT("Last month"))));
-	FromDateInDaysValues.Add(TPair<int32, FText>(90, FText::FromString(TEXT("Last 3 months"))));
+	FromDateInDaysValues.Add(TPair<int32, FText>(91, FText::FromString(TEXT("Last 3 months"))));
+	FromDateInDaysValues.Add(TPair<int32, FText>(182, FText::FromString(TEXT("Last 6 months"))));
 	FromDateInDaysValues.Add(TPair<int32, FText>(365, FText::FromString(TEXT("Last year"))));
 	FromDateInDaysValues.Add(TPair<int32, FText>(-1, FText::FromString(TEXT("All time"))));
 
@@ -62,45 +69,136 @@ void SPlasticSourceControlBranchesWidget::Construct(const FArguments& InArgs)
 			[
 				SNew(SHorizontalBox)
 				+SHorizontalBox::Slot()
-				.HAlign(HAlign_Left)
-				.VAlign(VAlign_Center)
-				.AutoWidth()
+				.FillWidth(1.0f)
 				[
-					CreateToolBar()
-				]
-				+SHorizontalBox::Slot()
-				.MaxWidth(10.0f)
-				[
-					SNew(SSpacer)
+					SNew(SHorizontalBox)
+					+SHorizontalBox::Slot()
+					.HAlign(HAlign_Left)
+					.VAlign(VAlign_Center)
+					.AutoWidth()
+					[
+						CreateToolBar()
+					]
+					+SHorizontalBox::Slot()
+					.MaxWidth(10.0f)
+					[
+						SNew(SSpacer)
+					]
+					+SHorizontalBox::Slot()
+					.VAlign(VAlign_Center)
+					.MaxWidth(300.0f)
+					[
+						SAssignNew(BranchSearchBox, SSearchBox)
+						.HintText(LOCTEXT("SearchBranches", "Search Branches"))
+						.ToolTipText(LOCTEXT("PlasticBranchesSearch_Tooltip", "Filter the list of branches by keyword."))
+						.OnTextChanged(this, &SPlasticSourceControlBranchesWidget::OnSearchTextChanged)
+					]
+					+SHorizontalBox::Slot()
+					.VAlign(VAlign_Center)
+					.MaxWidth(125.0f)
+					.Padding(10.f, 0.f)
+					[
+						SNew(SComboButton)
+						.ToolTipText(LOCTEXT("PlasticBranchesDate_Tooltip", "Filter the list of branches by date of activity."))
+						.OnGetMenuContent(this, &SPlasticSourceControlBranchesWidget::BuildFromDateDropDownMenu)
+						.ButtonContent()
+						[
+							SNew(STextBlock)
+							.Text_Lambda([this]() { return FromDateInDaysValues[FromDateInDays]; })
+						]
+					]
 				]
 				+SHorizontalBox::Slot()
+				.HAlign(HAlign_Right)
 				.VAlign(VAlign_Center)
-				.MaxWidth(300.0f)
+				.AutoWidth()
 				[
-					SAssignNew(FileSearchBox, SSearchBox)
-					.HintText(LOCTEXT("SearchBranches", "Search Branches"))
-					.ToolTipText(LOCTEXT("PlasticBranchesSearch_Tooltip", "Filter the list of branches by keyword."))
-					.OnTextChanged(this, &SPlasticSourceControlBranchesWidget::OnSearchTextChanged)
+					// Button to open the Changesets View
+					SNew(SButton)
+					.ContentPadding(FMargin(6.0f, 0.0f))
+					.ToolTipText(LOCTEXT("PlasticChangesetsWindowTooltip", "Open the Changesets window."))
+					.ButtonStyle(FAppStyle::Get(), "SimpleButton")
+					.OnClicked_Lambda([]()
+						{
+							FPlasticSourceControlModule::Get().GetChangesetsWindow().OpenTab();
+							return FReply::Handled();
+						})
+					[
+						SNew(SHorizontalBox)
+						+SHorizontalBox::Slot()
+						.AutoWidth()
+						.VAlign(VAlign_Center)
+						.HAlign(HAlign_Center)
+						[
+							SNew(SImage)
+							.Image(FAppStyle::GetBrush("SourceControl.Actions.History"))
+						]
+						+SHorizontalBox::Slot()
+						.AutoWidth()
+						.VAlign(VAlign_Center)
+						.Padding(5.0f, 0.0f, 0.0f, 0.0f)
+						[
+							SNew(STextBlock)
+							.TextStyle(&FAppStyle::Get().GetWidgetStyle<FTextBlockStyle>("NormalText"))
+							.Text(LOCTEXT("PlasticChangesetsWindow", "Changesets"))
+						]
+					]
 				]
 				+SHorizontalBox::Slot()
+				.HAlign(HAlign_Right)
 				.VAlign(VAlign_Center)
-				.MaxWidth(125.0f)
-				.Padding(10.f, 0.f)
+				.AutoWidth()
 				[
-					SNew(SComboButton)
-					.ToolTipText(LOCTEXT("PlasticBranchesDate_Tooltip", "Filter the list of branches by date of creation."))
-					.OnGetMenuContent(this, &SPlasticSourceControlBranchesWidget::BuildFromDateDropDownMenu)
-					.ButtonContent()
+					// Button to open the Branch Explorer
+					SNew(SButton)
+					.ContentPadding(FMargin(6.0f, 0.0f))
+					.ToolTipText(LOCTEXT("PlasticBranchExplorerTooltip", "Open the Branch Explorer of the Desktop Application for the current workspace."))
+					.ButtonStyle(FAppStyle::Get(), "SimpleButton")
+					.OnClicked_Lambda([]()
+						{
+							PlasticSourceControlUtils::OpenDesktopApplication(true);
+							return FReply::Handled();
+						})
 					[
-						SNew(STextBlock)
-						.Text_Lambda([this]() { return FromDateInDaysValues[FromDateInDays]; })
+						SNew(SHorizontalBox)
+						+SHorizontalBox::Slot()
+						.AutoWidth()
+						.VAlign(VAlign_Center)
+						.HAlign(HAlign_Center)
+						[
+							SNew(SImage)
+							.Image(FAppStyle::GetBrush("SourceControl.Branch"))
+						]
+						+SHorizontalBox::Slot()
+						.AutoWidth()
+						.VAlign(VAlign_Center)
+						.Padding(5.0f, 0.0f, 0.0f, 0.0f)
+						[
+							SNew(STextBlock)
+							.TextStyle(&FAppStyle::Get().GetWidgetStyle<FTextBlockStyle>("NormalText"))
+							.Text(LOCTEXT("OpenBranchExplorer", "Branch Explorer"))
+						]
 					]
 				]
 			]
 		]
 		+SVerticalBox::Slot() // The main content: the list of branches
 		[
-			CreateContentPanel()
+			SNew(SVerticalBox)
+			+SVerticalBox::Slot()
+			[
+				CreateContentPanel()
+			]
+			+SVerticalBox::Slot()
+			.VAlign(VAlign_Center)
+			.HAlign(HAlign_Center)
+			.FillHeight(1.0f)
+			[
+				// Text to display when there is no branch displayed
+				SNew(STextBlock)
+				.Text(LOCTEXT("NoBranch", "There is no branch to display."))
+				.Visibility_Lambda([this]() { return SourceControlBranches.Num() ? EVisibility::Collapsed : EVisibility::Visible; })
+			]
 		]
 		+SVerticalBox::Slot() // Status bar (Always visible)
 		.AutoHeight()
@@ -121,7 +219,7 @@ void SPlasticSourceControlBranchesWidget::Construct(const FArguments& InArgs)
 				.HAlign(HAlign_Right)
 				[
 					SNew(STextBlock)
-					.Text_Lambda([this]() { return FText::FromString(CurrentBranchName); })
+					.Text_Lambda([this]() { return FText::FromString(WorkspaceSelector); })
 					.ToolTipText(LOCTEXT("PlasticBranchCurrent_Tooltip", "Current branch."))
 				]
 			]
@@ -134,8 +232,7 @@ TSharedRef<SWidget> SPlasticSourceControlBranchesWidget::CreateToolBar()
 	FSlimHorizontalToolBarBuilder ToolBarBuilder(nullptr, FMultiBoxCustomization::None);
 
 	ToolBarBuilder.AddToolBarButton(
-		FUIAction(
-			FExecuteAction::CreateLambda([this]() { RequestBranchesRefresh(); })),
+		FUIAction(FExecuteAction::CreateLambda([this]() { bShouldRefresh = true; })),
 		NAME_None,
 		LOCTEXT("SourceControl_RefreshButton", "Refresh"),
 		LOCTEXT("SourceControl_RefreshButton_Tooltip", "Refreshes branches from revision control provider."),
@@ -172,6 +269,7 @@ TSharedRef<SWidget> SPlasticSourceControlBranchesWidget::CreateContentPanel()
 		.OnGenerateRow(this, &SPlasticSourceControlBranchesWidget::OnGenerateRow)
 		.SelectionMode(ESelectionMode::Multi)
 		.OnContextMenuOpening(this, &SPlasticSourceControlBranchesWidget::OnOpenContextMenu)
+		.OnMouseButtonDoubleClick(this, &SPlasticSourceControlBranchesWidget::OnItemDoubleClicked)
 		.OnItemToString_Debug_Lambda([this](FPlasticSourceControlBranchRef Branch) { return Branch->Name; })
 		.HeaderRow
 		(
@@ -229,11 +327,11 @@ TSharedRef<SWidget> SPlasticSourceControlBranchesWidget::CreateContentPanel()
 
 TSharedRef<ITableRow> SPlasticSourceControlBranchesWidget::OnGenerateRow(FPlasticSourceControlBranchRef InBranch, const TSharedRef<STableViewBase>& OwnerTable)
 {
-	const bool bIsCurrentBranch = InBranch->Name == CurrentBranchName;
+	const bool bIsCurrentBranch = InBranch->Name == WorkspaceSelector;
 	return SNew(SPlasticSourceControlBranchRow, OwnerTable)
 		.BranchToVisualize(InBranch)
 		.bIsCurrentBranch(bIsCurrentBranch)
-		.HighlightText_Lambda([this]() { return FileSearchBox->GetText(); });
+		.HighlightText_Lambda([this]() { return BranchSearchBox->GetText(); });
 }
 
 void SPlasticSourceControlBranchesWidget::OnHiddenColumnsListChanged()
@@ -242,7 +340,7 @@ void SPlasticSourceControlBranchesWidget::OnHiddenColumnsListChanged()
 	if (BranchesListView && BranchesListView->GetHeaderRow())
 	{
 		UPlasticSourceControlProjectSettings* Settings = GetMutableDefault<UPlasticSourceControlProjectSettings>();
-		Settings->bShowBranchRepositoryColumn = false;
+		Settings->bShowBranchRepositoryColumn = true;
 		Settings->bShowBranchCreatedByColumn = true;
 		Settings->bShowBranchDateColumn = true;
 		Settings->bShowBranchCommentColumn = true;
@@ -273,7 +371,7 @@ void SPlasticSourceControlBranchesWidget::OnHiddenColumnsListChanged()
 void SPlasticSourceControlBranchesWidget::OnSearchTextChanged(const FText& InFilterText)
 {
 	SearchTextFilter->SetRawFilterText(InFilterText);
-	FileSearchBox->SetError(SearchTextFilter->GetFilterErrorText());
+	BranchSearchBox->SetError(SearchTextFilter->GetFilterErrorText());
 }
 
 void SPlasticSourceControlBranchesWidget::PopulateItemSearchStrings(const FPlasticSourceControlBranch& InItem, TArray<FString>& OutStrings)
@@ -284,8 +382,7 @@ void SPlasticSourceControlBranchesWidget::PopulateItemSearchStrings(const FPlast
 void SPlasticSourceControlBranchesWidget::OnFromDateChanged(int32 InFromDateInDays)
 {
 	FromDateInDays = InFromDateInDays;
-
-	RequestBranchesRefresh();
+	bShouldRefresh = true;
 }
 
 TSharedRef<SWidget> SPlasticSourceControlBranchesWidget::BuildFromDateDropDownMenu()
@@ -393,7 +490,7 @@ void SPlasticSourceControlBranchesWidget::SortBranchView()
 
 	auto CompareRepository = [](const FPlasticSourceControlBranch* Lhs, const FPlasticSourceControlBranch* Rhs)
 	{
-		return FCString::Stricmp(*Lhs->Repository, *Rhs->Repository);
+		return UE::ComparisonUtility::CompareNaturalOrder(*Lhs->Repository, *Rhs->Repository);
 	};
 
 	auto CompareCreatedBy = [](const FPlasticSourceControlBranch* Lhs, const FPlasticSourceControlBranch* Rhs)
@@ -519,7 +616,7 @@ TSharedPtr<SWidget> SPlasticSourceControlBranchesWidget::OnOpenContextMenu()
 		return nullptr;
 	}
 	const bool bSingleSelection = !SelectedBranch.IsEmpty();
-	const bool bSingleNotCurrent = bSingleSelection && (SelectedBranch != CurrentBranchName);
+	const bool bSingleNotCurrent = bSingleSelection && (SelectedBranch != WorkspaceSelector);
 
 	const bool bMergeXml = FPlasticSourceControlModule::Get().GetProvider().GetPlasticScmVersion() >= PlasticSourceControlVersions::MergeXml;
 
@@ -546,7 +643,7 @@ TSharedPtr<SWidget> SPlasticSourceControlBranchesWidget::OnOpenContextMenu()
 		FText::FromString(SelectedBranch)));
 	const FText& CreateChildBranchTooltipDynamic = bSingleSelection ? CreateChildBranchTooltip : SelectASingleBranchTooltip;
 	Section.AddMenuEntry(
-		TEXT("CreateChildBranch"),
+		"CreateChildBranch",
 		LOCTEXT("CreateChildBranch", "Create child branch..."),
 		CreateChildBranchTooltipDynamic,
 		FSlateIcon(),
@@ -562,7 +659,7 @@ TSharedPtr<SWidget> SPlasticSourceControlBranchesWidget::OnOpenContextMenu()
 		bSingleNotCurrent ? SwitchToBranchTooltip :
 		bSingleSelection ? SelectADifferentBranchTooltip : SelectASingleBranchTooltip;
 	Section.AddMenuEntry(
-		TEXT("SwitchToBranch"),
+		"SwitchToBranch",
 		LOCTEXT("SwitchToBranch", "Switch workspace to this branch"),
 		SwitchToBranchTooltipDynamic,
 		FSlateIcon(),
@@ -575,13 +672,13 @@ TSharedPtr<SWidget> SPlasticSourceControlBranchesWidget::OnOpenContextMenu()
 	Section.AddSeparator("PlasticSeparator1");
 
 	const FText MergeBranchTooltip(FText::Format(LOCTEXT("MergeBranchTooltip", "Merge this branch {0} into the current branch {1}"),
-		FText::FromString(SelectedBranch), FText::FromString(CurrentBranchName)));
+		FText::FromString(SelectedBranch), FText::FromString(WorkspaceSelector)));
 	const FText& MergeBranchTooltipDynamic =
 		!bMergeXml ? UpdateUVCSTooltip :
 		bSingleNotCurrent ? MergeBranchTooltip :
 		bSingleSelection ? SelectADifferentBranchTooltip : SelectASingleBranchTooltip;
 	Section.AddMenuEntry(
-		TEXT("MergeBranch"),
+		"MergeBranch",
 		LOCTEXT("MergeBranch", "Merge from this branch..."),
 		MergeBranchTooltipDynamic,
 		FSlateIcon(),
@@ -597,7 +694,7 @@ TSharedPtr<SWidget> SPlasticSourceControlBranchesWidget::OnOpenContextMenu()
 		FText::FromString(SelectedBranch)));
 	const FText& RenameBranchTooltipDynamic = bSingleSelection ? RenameBranchTooltip : SelectASingleBranchTooltip;
 	Section.AddMenuEntry(
-		TEXT("RenameBranch"),
+		"RenameBranch",
 		LOCTEXT("RenameBranch", "Rename..."),
 		RenameBranchTooltipDynamic,
 		FSlateIcon(),
@@ -610,7 +707,7 @@ TSharedPtr<SWidget> SPlasticSourceControlBranchesWidget::OnOpenContextMenu()
 		FText::Format(LOCTEXT("DeleteBranchTooltip", "Delete the branch {0}"), FText::FromString(SelectedBranch)) :
 		LOCTEXT("DeleteBranchesTooltip", "Delete the selected branches.");
 	Section.AddMenuEntry(
-		TEXT("DeleteBranch"),
+		"DeleteBranch",
 		LOCTEXT("DeleteBranch", "Delete"),
 		DeleteBranchTooltip,
 		FSlateIcon(),
@@ -707,7 +804,7 @@ void SPlasticSourceControlBranchesWidget::OnSwitchToBranchClicked(FString InBran
 
 		// Launch a custom "Switch" operation
 		FPlasticSourceControlProvider& Provider = FPlasticSourceControlModule::Get().GetProvider();
-		TSharedRef<FPlasticSwitchToBranch, ESPMode::ThreadSafe> SwitchToBranchOperation = ISourceControlOperation::Create<FPlasticSwitchToBranch>();
+		TSharedRef<FPlasticSwitch, ESPMode::ThreadSafe> SwitchToBranchOperation = ISourceControlOperation::Create<FPlasticSwitch>();
 		SwitchToBranchOperation->BranchName = InBranchName;
 		const ECommandResult::Type Result = Provider.Execute(SwitchToBranchOperation, TArray<FString>(), EConcurrency::Asynchronous, FSourceControlOperationComplete::CreateSP(this, &SPlasticSourceControlBranchesWidget::OnSwitchToBranchOperationComplete));
 		if (Result == ECommandResult::Succeeded)
@@ -732,7 +829,7 @@ void SPlasticSourceControlBranchesWidget::OnSwitchToBranchClicked(FString InBran
 
 void SPlasticSourceControlBranchesWidget::OnMergeBranchClicked(FString InBranchName)
 {
-	const FText MergeBranchQuestion = FText::Format(LOCTEXT("MergeBranchDialog", "Merge branch {0} into the current branch {1}?"), FText::FromString(InBranchName), FText::FromString(CurrentBranchName));
+	const FText MergeBranchQuestion = FText::Format(LOCTEXT("MergeBranchDialog", "Merge branch {0} into the current branch {1}?"), FText::FromString(InBranchName), FText::FromString(WorkspaceSelector));
 	const EAppReturnType::Type Choice = FMessageDialog::Open(
 		EAppMsgCategory::Info,
 		EAppMsgType::YesNo, MergeBranchQuestion
@@ -955,7 +1052,7 @@ void SPlasticSourceControlBranchesWidget::OnGetBranchesOperationComplete(const F
 	TSharedRef<FPlasticGetBranches, ESPMode::ThreadSafe> OperationGetBranches = StaticCastSharedRef<FPlasticGetBranches>(InOperation);
 	SourceControlBranches = MoveTemp(OperationGetBranches->Branches);
 
-	CurrentBranchName = FPlasticSourceControlModule::Get().GetProvider().GetBranchName();
+	WorkspaceSelector = FPlasticSourceControlModule::Get().GetProvider().GetWorkspaceSelector();
 
 	EndRefreshStatus();
 	OnRefreshUI();
@@ -993,7 +1090,7 @@ void SPlasticSourceControlBranchesWidget::OnSwitchToBranchOperationComplete(cons
 	TRACE_CPUPROFILER_EVENT_SCOPE(SPlasticSourceControlBranchesWidget::OnSwitchToBranchOperationComplete);
 
 	// Reload packages that where updated by the SwitchToBranch operation (and the current map if needed)
-	TSharedRef<FPlasticSwitchToBranch, ESPMode::ThreadSafe> SwitchToBranchOperation = StaticCastSharedRef<FPlasticSwitchToBranch>(InOperation);
+	TSharedRef<FPlasticSwitch, ESPMode::ThreadSafe> SwitchToBranchOperation = StaticCastSharedRef<FPlasticSwitch>(InOperation);
 	PackageUtils::ReloadPackages(SwitchToBranchOperation->UpdatedFiles);
 
 	// Ask for a full refresh of the list of branches (and don't call EndRefreshStatus() yet)
@@ -1054,38 +1151,58 @@ void SPlasticSourceControlBranchesWidget::OnSourceControlProviderChanged(ISource
 	}
 }
 
+void SPlasticSourceControlBranchesWidget::SwitchToBranchWithConfirmation(const FString& InSelectedBranch)
+{
+	const FText SwitchToBranchQuestion = FText::Format(LOCTEXT("SwitchToBranchDialog", "Switch workspace to branch {0}?"), FText::FromString(InSelectedBranch));
+	const EAppReturnType::Type Choice = FMessageDialog::Open(
+		EAppMsgCategory::Info,
+		EAppMsgType::YesNo, SwitchToBranchQuestion
+		, LOCTEXT("SwitchToBranchTitle", "Switch Branch?")
+	);
+	if (Choice == EAppReturnType::Yes)
+	{
+		OnSwitchToBranchClicked(InSelectedBranch);
+	}
+}
+
+void SPlasticSourceControlBranchesWidget::HandleSourceControlStateChanged()
+{
+	if (WorkspaceSelector != FPlasticSourceControlModule::Get().GetProvider().GetWorkspaceSelector())
+	{
+		bShouldRefresh = true;
+	}
+}
+
+void SPlasticSourceControlBranchesWidget::OnItemDoubleClicked(FPlasticSourceControlBranchRef InBranch)
+{
+	// Double click switches to the selected branch (with a confirmation dialog)
+	if (InBranch->Name != WorkspaceSelector)
+	{
+		SwitchToBranchWithConfirmation(InBranch->Name);
+	}
+}
+
 FReply SPlasticSourceControlBranchesWidget::OnKeyDown(const FGeometry& MyGeometry, const FKeyEvent& InKeyEvent)
 {
 	if (InKeyEvent.GetKey() == EKeys::F5)
 	{
 		// Pressing F5 refreshes the list of branches
-		RequestBranchesRefresh();
+		bShouldRefresh = true;
 		return FReply::Handled();
 	}
 	else if (InKeyEvent.GetKey() == EKeys::Enter)
 	{
-		// Pressing Enter switches to the selected branch.
+		// Pressing Enter switches to the selected branch (with a confirmation dialog)
 		const TArray<FString> SelectedBranches = GetSelectedBranches();
-		if (SelectedBranches.Num() == 1)
+		if (SelectedBranches.Num() == 1 && SelectedBranches[0] != WorkspaceSelector)
 		{
-			const FString& SelectedBranch = SelectedBranches[0];
-			// Note: this action require a confirmation dialog (while the Delete below already have one in OnDeleteBranchesClicked()).
-			const FText SwitchToBranchQuestion = FText::Format(LOCTEXT("SwitchToBranchDialog", "Switch workspace to branch {0}?"), FText::FromString(SelectedBranch));
-			const EAppReturnType::Type Choice = FMessageDialog::Open(
-				EAppMsgCategory::Info,
-				EAppMsgType::YesNo, SwitchToBranchQuestion
-				, LOCTEXT("SwitchToBranchTitle", "Switch Branch?")
-			);
-			if (Choice == EAppReturnType::Yes)
-			{
-				OnSwitchToBranchClicked(SelectedBranch);
-			}
+			SwitchToBranchWithConfirmation(SelectedBranches[0]);
 		}
 		return FReply::Handled();
 	}
 	else if (InKeyEvent.GetKey() == EKeys::F2)
 	{
-		// Pressing F2 renames the selected branches
+		// Pressing F2 renames the selected branches (with a dialog)
 		const TArray<FString> SelectedBranches = GetSelectedBranches();
 		if (SelectedBranches.Num() == 1)
 		{
@@ -1096,7 +1213,7 @@ FReply SPlasticSourceControlBranchesWidget::OnKeyDown(const FGeometry& MyGeometr
 	}
 	else if ((InKeyEvent.GetKey() == EKeys::Delete) || (InKeyEvent.GetKey() == EKeys::BackSpace))
 	{
-		// Pressing Delete or BackSpace deletes the selected branches
+		// Pressing Delete or BackSpace deletes the selected branches (with a confirmation dialog)
 		const TArray<FString> SelectedBranches = GetSelectedBranches();
 		if (SelectedBranches.Num() > 0)
 		{
@@ -1108,4 +1225,4 @@ FReply SPlasticSourceControlBranchesWidget::OnKeyDown(const FGeometry& MyGeometr
 	return FReply::Unhandled();
 }
 
-#undef LOCTEXT_NAMESPACE
\ No newline at end of file
+#undef LOCTEXT_NAMESPACE
diff --git a/Engine/Plugins/Developer/PlasticSourceControl/Source/PlasticSourceControl/Private/SPlasticSourceControlBranchesWidget.h b/Engine/Plugins/Developer/PlasticSourceControl/Source/PlasticSourceControl/Private/SPlasticSourceControlBranchesWidget.h
--- a/Engine/Plugins/Developer/PlasticSourceControl/Source/PlasticSourceControl/Private/SPlasticSourceControlBranchesWidget.h
+++ b/Engine/Plugins/Developer/PlasticSourceControl/Source/PlasticSourceControl/Private/SPlasticSourceControlBranchesWidget.h
@@ -55,6 +55,7 @@ private:
 	void OnColumnSortModeChanged(const EColumnSortPriority::Type InSortPriority, const FName& InColumnId, const EColumnSortMode::Type InSortMode);
 
 	void SortBranchView();
+
 	TArray<FString> GetSelectedBranches();
 
 	TSharedPtr<SWidget> OnOpenContextMenu();
@@ -84,16 +85,24 @@ private:
 	void OnDeleteBranchesOperationComplete(const FSourceControlOperationRef& InOperation, ECommandResult::Type InResult);
 	void OnSourceControlProviderChanged(ISourceControlProvider& OldProvider, ISourceControlProvider& NewProvider);
 
+	/** Delegate handler for when source control state changes */
+	void HandleSourceControlStateChanged();
+
 	SListView<FPlasticSourceControlBranchRef>* GetListView() const
 	{
 		return BranchesListView.Get();
 	}
 
+	void SwitchToBranchWithConfirmation(const FString& InSelectedBranch);
+
+	/** Double click to switch to the branch */
+	void OnItemDoubleClicked(FPlasticSourceControlBranchRef InBranch);
+
 	/** Interpret F5, Enter and Delete keys */
 	virtual FReply OnKeyDown(const FGeometry& MyGeometry, const FKeyEvent& InKeyEvent) override;
 
 private:
-	TSharedPtr<SSearchBox> FileSearchBox;
+	TSharedPtr<SSearchBox> BranchSearchBox;
 
 	FName PrimarySortedColumn;
 	FName SecondarySortedColumn;
@@ -109,7 +118,7 @@ private:
 	bool bIsRefreshing = false;
 	double RefreshStatusStartSecs;
 
-	FString CurrentBranchName;
+	FString WorkspaceSelector;
 
 	/** Ongoing notification for a long-running asynchronous source control operation, if any */
 	FNotification Notification;
@@ -123,6 +132,9 @@ private:
 	TArray<FPlasticSourceControlBranchRef> SourceControlBranches; // Full list from source (filtered by date)
 	TArray<FPlasticSourceControlBranchRef> BranchRows; // Filtered list to display based on the search text filter
 
+	/** Delegate handle for the HandleSourceControlStateChanged function callback */
+	FDelegateHandle SourceControlStateChangedDelegateHandle;
+
 	/** The dialog Window that opens when the user click on any context menu entry */
 	TSharedPtr<SWindow> DialogWindowPtr;
-};
\ No newline at end of file
+};
diff --git a/Engine/Plugins/Developer/PlasticSourceControl/Source/PlasticSourceControl/Private/SPlasticSourceControlCreateBranch.cpp b/Engine/Plugins/Developer/PlasticSourceControl/Source/PlasticSourceControl/Private/SPlasticSourceControlCreateBranch.cpp
--- a/Engine/Plugins/Developer/PlasticSourceControl/Source/PlasticSourceControl/Private/SPlasticSourceControlCreateBranch.cpp
+++ b/Engine/Plugins/Developer/PlasticSourceControl/Source/PlasticSourceControl/Private/SPlasticSourceControlCreateBranch.cpp
@@ -11,10 +11,11 @@
 #include "Widgets/Input/SCheckBox.h"
 #include "Widgets/Input/SEditableTextBox.h"
 #include "Widgets/Input/SMultiLineEditableTextBox.h"
+#include "Widgets/Text/STextBlock.h"
 #include "Widgets/SBoxPanel.h"
 #include "Widgets/SWindow.h"
 
-#define LOCTEXT_NAMESPACE "PlasticSourceControlWindow"
+#define LOCTEXT_NAMESPACE "PlasticSourceControlBranchesWindow"
 
 void SPlasticSourceControlCreateBranch::Construct(const FArguments& InArgs)
 {
@@ -220,4 +221,4 @@ FReply SPlasticSourceControlCreateBranch::OnKeyDown(const FGeometry& MyGeometry,
 	return FReply::Unhandled();
 }
 
-#undef LOCTEXT_NAMESPACE
\ No newline at end of file
+#undef LOCTEXT_NAMESPACE
diff --git a/Engine/Plugins/Developer/PlasticSourceControl/Source/PlasticSourceControl/Private/SPlasticSourceControlDeleteBranches.cpp b/Engine/Plugins/Developer/PlasticSourceControl/Source/PlasticSourceControl/Private/SPlasticSourceControlDeleteBranches.cpp
--- a/Engine/Plugins/Developer/PlasticSourceControl/Source/PlasticSourceControl/Private/SPlasticSourceControlDeleteBranches.cpp
+++ b/Engine/Plugins/Developer/PlasticSourceControl/Source/PlasticSourceControl/Private/SPlasticSourceControlDeleteBranches.cpp
@@ -9,9 +9,11 @@
 #include "Widgets/Input/SButton.h"
 #include "Widgets/Input/SEditableTextBox.h"
 #include "Widgets/Input/SMultiLineEditableTextBox.h"
+#include "Widgets/Text/STextBlock.h"
+#include "Widgets/SBoxPanel.h"
 #include "Widgets/SWindow.h"
 
-#define LOCTEXT_NAMESPACE "PlasticSourceControlWindow"
+#define LOCTEXT_NAMESPACE "PlasticSourceControlBranchesWindow"
 
 void SPlasticSourceControlDeleteBranches::Construct(const FArguments& InArgs)
 {
@@ -127,4 +129,4 @@ FReply SPlasticSourceControlDeleteBranches::OnKeyDown(const FGeometry& MyGeometr
 	return FReply::Unhandled();
 }
 
-#undef LOCTEXT_NAMESPACE
\ No newline at end of file
+#undef LOCTEXT_NAMESPACE
diff --git a/Engine/Plugins/Developer/PlasticSourceControl/Source/PlasticSourceControl/Private/SPlasticSourceControlRenameBranch.cpp b/Engine/Plugins/Developer/PlasticSourceControl/Source/PlasticSourceControl/Private/SPlasticSourceControlRenameBranch.cpp
--- a/Engine/Plugins/Developer/PlasticSourceControl/Source/PlasticSourceControl/Private/SPlasticSourceControlRenameBranch.cpp
+++ b/Engine/Plugins/Developer/PlasticSourceControl/Source/PlasticSourceControl/Private/SPlasticSourceControlRenameBranch.cpp
@@ -9,9 +9,11 @@
 #include "Widgets/Input/SButton.h"
 #include "Widgets/Input/SEditableTextBox.h"
 #include "Widgets/Input/SMultiLineEditableTextBox.h"
+#include "Widgets/Text/STextBlock.h"
+#include "Widgets/SBoxPanel.h"
 #include "Widgets/SWindow.h"
 
-#define LOCTEXT_NAMESPACE "PlasticSourceControlWindow"
+#define LOCTEXT_NAMESPACE "PlasticSourceControlBranchesWindow"
 
 // Extract the short name of the branch, after the last slash
 static FString GetShortBranchName(const FString& InBranchName)
@@ -198,4 +200,4 @@ FReply SPlasticSourceControlRenameBranch::OnKeyDown(const FGeometry& MyGeometry,
 	return FReply::Unhandled();
 }
 
-#undef LOCTEXT_NAMESPACE
\ No newline at end of file
+#undef LOCTEXT_NAMESPACE
diff --git a/Engine/Plugins/Developer/PlasticSourceControl/Source/PlasticSourceControl/Private/SPlasticSourceControlSettings.cpp b/Engine/Plugins/Developer/PlasticSourceControl/Source/PlasticSourceControl/Private/SPlasticSourceControlSettings.cpp
--- a/Engine/Plugins/Developer/PlasticSourceControl/Source/PlasticSourceControl/Private/SPlasticSourceControlSettings.cpp
+++ b/Engine/Plugins/Developer/PlasticSourceControl/Source/PlasticSourceControl/Private/SPlasticSourceControlSettings.cpp
@@ -15,11 +15,12 @@
 #include "Widgets/Input/SButton.h"
 #include "Widgets/Input/SCheckBox.h"
 #include "Widgets/Input/SEditableTextBox.h"
+#include "Widgets/Input/SHyperlink.h"
 #include "Widgets/Input/SMultiLineEditableTextBox.h"
 #include "Widgets/Layout/SBorder.h"
 #include "Widgets/Layout/SSeparator.h"
-#include "Widgets/SBoxPanel.h"
 #include "Widgets/Text/STextBlock.h"
+#include "Widgets/SBoxPanel.h"
 #include "Styling/AppStyle.h"
 
 #define LOCTEXT_NAMESPACE "SPlasticSourceControlSettings"
@@ -66,9 +67,9 @@ void SPlasticSourceControlSettings::Construct(const FArguments& InArgs)
 				.Font(Font)
 			]
 		]
-		// Unity Version Control command line tool not available warning
+		// Unity Version Control command line tool not available warning and download link
 		+SVerticalBox::Slot()
-		.FillHeight(1.0f)
+		.AutoHeight()
 		.Padding(2.0f)
 		[
 			SNew(STextBlock)
@@ -78,6 +79,30 @@ void SPlasticSourceControlSettings::Construct(const FArguments& InArgs)
 			.Font(Font)
 			.ColorAndOpacity(FLinearColor::Red)
 		]
+		+SVerticalBox::Slot()
+		.AutoHeight()
+		.Padding(2.0f)
+		[
+			SNew(SHorizontalBox)
+			+SHorizontalBox::Slot()
+			.AutoWidth()
+			[
+				SNew(SHyperlink)
+				.Visibility(this, &SPlasticSourceControlSettings::PlasticNotAvailable)
+				.ToolTipText(LOCTEXT("PlasticDownload_Tooltip", "Download Unity Version Control (Plastic SCM)"))
+				.Text(LOCTEXT("PlasticDownload", "https://www.plasticscm.com/download/downloadinstaller/..."))
+				.OnNavigate(FSimpleDelegate::CreateLambda([]()
+				{
+#if PLATFORM_WINDOWS
+					FPlatformProcess::LaunchURL(TEXT("https://www.plasticscm.com/download/downloadinstaller/last/plasticscm/windows/cloudedition"), NULL, NULL);
+#elif PLATFORM_MAC
+					FPlatformProcess::LaunchURL(TEXT("https://www.plasticscm.com/download/downloadinstaller/last/plasticscm/macosx/cloudedition"), NULL, NULL);
+#elif PLATFORM_LINUX
+					FPlatformProcess::LaunchURL(TEXT("https://www.plasticscm.com/plastic-for-linux"), NULL, NULL);
+#endif
+				}))
+			]
+		]
 		// Path to the Unity Version Control binary
 		+SVerticalBox::Slot()
 		.AutoHeight()
@@ -159,7 +184,7 @@ void SPlasticSourceControlSettings::Construct(const FArguments& InArgs)
 		]
 		// Explanation text
 		+SVerticalBox::Slot()
-		.FillHeight(1.0f)
+		.AutoHeight()
 		.Padding(2.0f, 5.0f)
 		[
 			SNew(STextBlock)
@@ -332,6 +357,22 @@ void SPlasticSourceControlSettings::Construct(const FArguments& InArgs)
 				.Font(Font)
 			]
 		]
+		// Option for the View Changes (Changelists) window to also show locally Changed and Private files
+		+SVerticalBox::Slot()
+		.AutoHeight()
+		.Padding(2.0f)
+		.VAlign(VAlign_Center)
+		[
+			SNew(SCheckBox)
+			.ToolTipText(LOCTEXT("ViewLocalChanges_Tooltip", "Enable the \"View Changes\" window to search for and show locally Changed and Private files (can be slow)."))
+			.IsChecked(SPlasticSourceControlSettings::IsViewLocalChangesChecked())
+			.OnCheckStateChanged(this, &SPlasticSourceControlSettings::OnCheckedViewLocalChanges)
+			[
+				SNew(STextBlock)
+				.Text(LOCTEXT("ViewLocalChanges", "Show local Changes in the \"View Changes\" window."))
+				.Font(Font)
+			]
+		]
 		// Option to enable Source Control Verbose logs
 		+SVerticalBox::Slot()
 		.AutoHeight()
@@ -585,6 +626,19 @@ ECheckBoxState SPlasticSourceControlSettings::IsUpdateStatusOtherBranchesChecked
 	return PlasticSettings.GetUpdateStatusOtherBranches() ? ECheckBoxState::Checked : ECheckBoxState::Unchecked;
 }
 
+void SPlasticSourceControlSettings::OnCheckedViewLocalChanges(ECheckBoxState NewCheckedState)
+{
+	FPlasticSourceControlSettings& PlasticSettings = FPlasticSourceControlModule::Get().GetProvider().AccessSettings();
+	PlasticSettings.SetViewLocalChanges(NewCheckedState == ECheckBoxState::Checked);
+	PlasticSettings.SaveSettings();
+}
+
+ECheckBoxState SPlasticSourceControlSettings::IsViewLocalChangesChecked() const
+{
+	FPlasticSourceControlSettings& PlasticSettings = FPlasticSourceControlModule::Get().GetProvider().AccessSettings();
+	return PlasticSettings.GetViewLocalChanges() ? ECheckBoxState::Checked : ECheckBoxState::Unchecked;
+}
+
 void SPlasticSourceControlSettings::OnCheckedEnableVerboseLogs(ECheckBoxState NewCheckedState)
 {
 	FPlasticSourceControlSettings& PlasticSettings = FPlasticSourceControlModule::Get().GetProvider().AccessSettings();
diff --git a/Engine/Plugins/Developer/PlasticSourceControl/Source/PlasticSourceControl/Private/SPlasticSourceControlSettings.h b/Engine/Plugins/Developer/PlasticSourceControl/Source/PlasticSourceControl/Private/SPlasticSourceControlSettings.h
--- a/Engine/Plugins/Developer/PlasticSourceControl/Source/PlasticSourceControl/Private/SPlasticSourceControlSettings.h
+++ b/Engine/Plugins/Developer/PlasticSourceControl/Source/PlasticSourceControl/Private/SPlasticSourceControlSettings.h
@@ -56,6 +56,9 @@ private:
 	void OnCheckedUpdateStatusOtherBranches(ECheckBoxState NewCheckedState);
 	ECheckBoxState IsUpdateStatusOtherBranchesChecked() const;
 
+	void OnCheckedViewLocalChanges(ECheckBoxState NewCheckedState);
+	ECheckBoxState IsViewLocalChangesChecked() const;
+
 	void OnCheckedEnableVerboseLogs(ECheckBoxState NewCheckedState);
 	ECheckBoxState IsEnableVerboseLogsChecked() const;
 
diff --git a/Engine/Plugins/Developer/PlasticSourceControl/Source/PlasticSourceControl/Private/SPlasticSourceControlStatusBar.cpp b/Engine/Plugins/Developer/PlasticSourceControl/Source/PlasticSourceControl/Private/SPlasticSourceControlStatusBar.cpp
--- a/Engine/Plugins/Developer/PlasticSourceControl/Source/PlasticSourceControl/Private/SPlasticSourceControlStatusBar.cpp
+++ b/Engine/Plugins/Developer/PlasticSourceControl/Source/PlasticSourceControl/Private/SPlasticSourceControlStatusBar.cpp
@@ -48,7 +48,7 @@ void SPlasticSourceControlStatusBar::Construct(const FArguments& InArgs)
 
 FText SPlasticSourceControlStatusBar::GetStatusBarText() const
 {;
-	return FText::FromString(FPlasticSourceControlModule::Get().GetProvider().GetBranchName());
+	return FText::FromString(FPlasticSourceControlModule::Get().GetProvider().GetWorkspaceSelector());
 }
 
 FReply SPlasticSourceControlStatusBar::OnClicked()
@@ -58,4 +58,4 @@ FReply SPlasticSourceControlStatusBar::OnClicked()
 	return FReply::Handled();
 }
 
-#undef LOCTEXT_NAMESPACE
\ No newline at end of file
+#undef LOCTEXT_NAMESPACE
diff --git a/Engine/Plugins/Developer/PlasticSourceControl/Source/PlasticSourceControl/Private/ScopedTempFile.cpp b/Engine/Plugins/Developer/PlasticSourceControl/Source/PlasticSourceControl/Private/ScopedTempFile.cpp
--- a/Engine/Plugins/Developer/PlasticSourceControl/Source/PlasticSourceControl/Private/ScopedTempFile.cpp
+++ b/Engine/Plugins/Developer/PlasticSourceControl/Source/PlasticSourceControl/Private/ScopedTempFile.cpp
@@ -8,27 +8,28 @@
 
 #include "ISourceControlModule.h" // LogSourceControl
 
-FScopedTempFile::FScopedTempFile()
+FScopedTempFile::FScopedTempFile() :
+	FScopedTempFile(TEXT("Temp-"), TEXT(".xml"))
 {
-	Filename = FPaths::CreateTempFilename(*FPaths::ConvertRelativePathToFull(FPaths::ProjectLogDir()), TEXT("Plastic-Temp"), TEXT(".txt"));
 }
 
-FScopedTempFile::FScopedTempFile(const FString& InText)
+FScopedTempFile::FScopedTempFile(const FString& InPrefix, const FString& InExtension)
+{
+	Filename = FPaths::CreateTempFilename(*FPaths::ConvertRelativePathToFull(FPaths::ProjectLogDir()), *InPrefix, *InExtension);
+}
+
+FScopedTempFile::FScopedTempFile(const FString& InText) :
+	FScopedTempFile(TEXT("Temp-"), TEXT(".txt"))
 {
-	Filename = FPaths::CreateTempFilename(*FPaths::ConvertRelativePathToFull(FPaths::ProjectLogDir()), TEXT("Plastic-Temp"), TEXT(".txt"));
 	if (!FFileHelper::SaveStringToFile(InText, *Filename, FFileHelper::EEncodingOptions::ForceUTF8WithoutBOM))
 	{
 		UE_LOG(LogSourceControl, Error, TEXT("Failed to write to temp file: %s"), *Filename);
 	}
 }
 
-FScopedTempFile::FScopedTempFile(const FText& InText)
+FScopedTempFile::FScopedTempFile(const FText& InText) :
+	FScopedTempFile(InText.ToString())
 {
-	Filename = FPaths::CreateTempFilename(*FPaths::ConvertRelativePathToFull(FPaths::ProjectLogDir()), TEXT("Plastic-Temp"), TEXT(".txt"));
-	if (!FFileHelper::SaveStringToFile(InText.ToString(), *Filename, FFileHelper::EEncodingOptions::ForceUTF8WithoutBOM))
-	{
-		UE_LOG(LogSourceControl, Error, TEXT("Failed to write to temp file: %s"), *Filename);
-	}
 }
 
 FScopedTempFile::~FScopedTempFile()
diff --git a/Engine/Plugins/Developer/PlasticSourceControl/Source/PlasticSourceControl/Private/ScopedTempFile.h b/Engine/Plugins/Developer/PlasticSourceControl/Source/PlasticSourceControl/Private/ScopedTempFile.h
--- a/Engine/Plugins/Developer/PlasticSourceControl/Source/PlasticSourceControl/Private/ScopedTempFile.h
+++ b/Engine/Plugins/Developer/PlasticSourceControl/Source/PlasticSourceControl/Private/ScopedTempFile.h
@@ -12,6 +12,7 @@ class FScopedTempFile
 public:
 	/** Default constructor - only hold a temp filename */
 	FScopedTempFile();
+	FScopedTempFile(const FString& InPrefix, const FString& InExtension);
 
 	/** Constructor - open & write string to temp file */
 	explicit FScopedTempFile(const FString& InText);
