[Revision Control] Fix another dupe locstring key in plastic notifications

#rnx

diff --git a/Engine/Plugins/Developer/PlasticSourceControl/Source/PlasticSourceControl/Private/Notification.cpp b/Engine/Plugins/Developer/PlasticSourceControl/Source/PlasticSourceControl/Private/Notification.cpp
--- a/Engine/Plugins/Developer/PlasticSourceControl/Source/PlasticSourceControl/Private/Notification.cpp
+++ b/Engine/Plugins/Developer/PlasticSourceControl/Source/PlasticSourceControl/Private/Notification.cpp
@@ -112,7 +112,7 @@ void FNotification::DisplayFailure(const FSourceControlOperationBase& InOperatio
 	{
 		// If there are multiple messages, display the last one to not let the user with a notification starting with a "wait" or "in progress" message
 		const FText NotificationText = FText::Format(
-			LOCTEXT("PlasticSourceControlOperation_Failure", "Error: {0} operation failed!\n{1}"),
+			LOCTEXT("PlasticSourceControlOperation_FailureWithMessages", "Error: {0} operation failed!\n{1}"),
 			FText::FromName(InOperation.GetName()),
 			InOperation.GetResultInfo().ErrorMessages.Last()
 		);
