RevisionControl / Buildfix: localization issues fixed introduced in #39402991. Fixes provided by original creator of the PR via GH.

#rnx

diff --git a/Engine/Plugins/Developer/PlasticSourceControl/Source/PlasticSourceControl/Private/PlasticSourceControlState.cpp b/Engine/Plugins/Developer/PlasticSourceControl/Source/PlasticSourceControl/Private/PlasticSourceControlState.cpp
--- a/Engine/Plugins/Developer/PlasticSourceControl/Source/PlasticSourceControl/Private/PlasticSourceControlState.cpp
+++ b/Engine/Plugins/Developer/PlasticSourceControl/Source/PlasticSourceControl/Private/PlasticSourceControlState.cpp
@@ -22,8 +22,8 @@ const TCHAR* FPlasticSourceControlState::ToString() const
 		case EWorkspaceState::Moved: return TEXT("Moved");
 		case EWorkspaceState::Copied: return TEXT("Copied");
 		case EWorkspaceState::Replaced: return TEXT("Replaced");
-		case EWorkspaceState::Deleted: return TEXT("Deleted");
-		case EWorkspaceState::LocallyDeleted: return TEXT("LocallyDeleted");
+		case EWorkspaceState::Deleted: return TEXT("Removed");
+		case EWorkspaceState::LocallyDeleted: return TEXT("Deleted locally");
 		case EWorkspaceState::Changed: return TEXT("Changed");
 		case EWorkspaceState::Conflicted: return TEXT("Conflicted");
 		case EWorkspaceState::Private: return TEXT("Private");
@@ -42,8 +42,8 @@ FText FPlasticSourceControlState::ToText() const
 	static const FText Moved = LOCTEXT("Moved", "Moved");
 	static const FText Copied = LOCTEXT("Copied", "Copied");
 	static const FText Replaced = LOCTEXT("Replaced", "Replaced");
-	static const FText Deleted = LOCTEXT("Deleted", "Deleted");
-	static const FText LocallyDeleted = LOCTEXT("LocallyDeleted", "LocallyDeleted");
+	static const FText Deleted = LOCTEXT("Deleted", "Removed");
+	static const FText LocallyDeleted = LOCTEXT("LocallyDeleted", "Deleted locally");
 	static const FText Changed = LOCTEXT("Changed", "Changed");
 	static const FText Conflicted = LOCTEXT("Conflicted", "Conflicted");
 	static const FText Private = LOCTEXT("Private", "Private");
@@ -585,4 +585,4 @@ bool FPlasticSourceControlState::CanRevert() const
 	return IsModified() || WorkspaceState == EWorkspaceState::CheckedOutUnchanged;
 }
 
-#undef LOCTEXT_NAMESPACE
+#undef LOCTEXT_NAMESPACE
\ No newline at end of file
diff --git a/Engine/Plugins/Developer/PlasticSourceControl/Source/PlasticSourceControl/Private/SPlasticSourceControlChangesetsWidget.cpp b/Engine/Plugins/Developer/PlasticSourceControl/Source/PlasticSourceControl/Private/SPlasticSourceControlChangesetsWidget.cpp
--- a/Engine/Plugins/Developer/PlasticSourceControl/Source/PlasticSourceControl/Private/SPlasticSourceControlChangesetsWidget.cpp
+++ b/Engine/Plugins/Developer/PlasticSourceControl/Source/PlasticSourceControl/Private/SPlasticSourceControlChangesetsWidget.cpp
@@ -980,7 +980,7 @@ TSharedPtr<SWidget> SPlasticSourceControlChangesetsWidget::OnOpenChangesetContex
 	Section.AddMenuEntry(
 		"DiffChangesets",
 		LOCTEXT("DiffChangesets", "Diff selected changesets"),
-		bDoubleSelection ? LOCTEXT("DiffChangesetTooltip", "Launch the Desktop application diff window showing changes between the two selected changesets.") : LOCTEXT("DoubleSelection", "Select a couple of changesets."),
+		bDoubleSelection ? LOCTEXT("DiffChangesetsTooltip", "Launch the Desktop application diff window showing changes between the two selected changesets.") : LOCTEXT("DoubleSelection", "Select a couple of changesets."),
 		FSlateIcon(),
 		FUIAction(
 			FExecuteAction::CreateSP(this, &SPlasticSourceControlChangesetsWidget::OnDiffChangesetsClicked, SelectedChangesets),
@@ -993,7 +993,7 @@ TSharedPtr<SWidget> SPlasticSourceControlChangesetsWidget::OnOpenChangesetContex
 	Section.AddMenuEntry(
 		"DiffBranch",
 		bSingleBranchSelected ? FText::Format(LOCTEXT("DiffBranchDynamic", "Diff branch {0}"), FText::FromString(SelectedChangeset->Branch)) : LOCTEXT("DiffBranch", "Diff branch"),
-		bSingleBranchSelected ? LOCTEXT("DiffChangesetTooltip", "Launch the Desktop application diff window showing all changes in the selected branch.") : SelectASingleBranchTooltip,
+		bSingleBranchSelected ? LOCTEXT("DiffBranchTooltip", "Launch the Desktop application diff window showing all changes in the selected branch.") : SelectASingleBranchTooltip,
 		FSlateIcon(),
 		FUIAction(
 			FExecuteAction::CreateSP(this, &SPlasticSourceControlChangesetsWidget::OnDiffBranchClicked, SelectedChangeset),
@@ -1951,4 +1951,4 @@ FReply SPlasticSourceControlChangesetsWidget::OnKeyDown(const FGeometry& MyGeome
 	return FReply::Unhandled();
 }
 
-#undef LOCTEXT_NAMESPACE
+#undef LOCTEXT_NAMESPACE
\ No newline at end of file
